Fix CIS, the arithmetic operator order.

diff --git a/Engine/Plugins/Experimental/Text3D/Source/Text3D/Private/Text3DComponent.cpp b/Engine/Plugins/Experimental/Text3D/Source/Text3D/Private/Text3DComponent.cpp
--- a/Engine/Plugins/Experimental/Text3D/Source/Text3D/Private/Text3DComponent.cpp
+++ b/Engine/Plugins/Experimental/Text3D/Source/Text3D/Private/Text3DComponent.cpp
@@ -907,7 +907,8 @@ void UText3DComponent::BuildTextMesh()
 						for (int32 Index = 0; Index < PointCount - 2; Index++)
 						{
 							const int32 TriangleStart = StartVertex + Index;
-							MeshesData->AddTriangle(TriangleStart + (Index % 2 ? 0 : 1), TriangleStart + (Index % 2 ? 1 : 0), TriangleStart + 2);
+							const bool bIsOdd = (Index % 2) != 0;
+							MeshesData->AddTriangle(TriangleStart + (bIsOdd ? 0 : 1), TriangleStart + (bIsOdd ? 1 : 0), TriangleStart + 2);
 						}
 
 						break;
