Copy up //UE4/Dev-Enterprise/...at  Changelist 10508152  to //UE4/Dev-Tools-Staging/...

diff --git a/Engine/Plugins/Enterprise/DatasmithImporter/Source/DatasmithImporter/Private/DatasmithBlueprintLibrary.cpp b/Engine/Plugins/Enterprise/DatasmithImporter/Source/DatasmithImporter/Private/DatasmithBlueprintLibrary.cpp
--- a/Engine/Plugins/Enterprise/DatasmithImporter/Source/DatasmithImporter/Private/DatasmithBlueprintLibrary.cpp
+++ b/Engine/Plugins/Enterprise/DatasmithImporter/Source/DatasmithImporter/Private/DatasmithBlueprintLibrary.cpp
@@ -4,14 +4,14 @@
 
 #include "DatasmithImportFactory.h"
 #include "DatasmithImportOptions.h"
-#include "DatasmithMeshHelper.h"
 #include "DatasmithSceneActor.h"
 #include "DatasmithSceneFactory.h"
 #include "DatasmithSceneSource.h"
 #include "DatasmithStaticMeshImporter.h"
+#include "DatasmithTranslatableSource.h"
 #include "ObjectElements/DatasmithUSceneElement.h"
-#include "Translators/DatasmithTranslatableSource.h"
 #include "Utility/DatasmithImporterUtils.h"
+#include "Utility/DatasmithMeshHelper.h"
 
 #include "Async/ParallelFor.h"
 #include "Editor.h"
diff --git a/Engine/Plugins/Enterprise/DatasmithImporter/Source/DatasmithImporter/Private/DatasmithConsumer.cpp b/Engine/Plugins/Enterprise/DatasmithImporter/Source/DatasmithImporter/Private/DatasmithConsumer.cpp
--- a/Engine/Plugins/Enterprise/DatasmithImporter/Source/DatasmithImporter/Private/DatasmithConsumer.cpp
+++ b/Engine/Plugins/Enterprise/DatasmithImporter/Source/DatasmithImporter/Private/DatasmithConsumer.cpp
@@ -307,6 +307,9 @@ bool UDatasmithConsumer::BuildContexts( UWorld* ImportWorld )
 		return false;
 	}
 
+	// Set the feedback context
+	ImportContextPtr->FeedbackContext = Context.ProgressReporterPtr ? Context.ProgressReporterPtr->GetFeedbackContext() : nullptr;
+
 	// Update ImportContext's package data
 	ImportContextPtr->AssetsContext.RootFolderPath = TargetContentFolder;
 	ImportContextPtr->AssetsContext.TransientFolderPath = Context.TransientContentFolder;
diff --git a/Engine/Plugins/Enterprise/DatasmithImporter/Source/DatasmithImporter/Private/DatasmithFileProducer.cpp b/Engine/Plugins/Enterprise/DatasmithImporter/Source/DatasmithImporter/Private/DatasmithFileProducer.cpp
--- a/Engine/Plugins/Enterprise/DatasmithImporter/Source/DatasmithImporter/Private/DatasmithFileProducer.cpp
+++ b/Engine/Plugins/Enterprise/DatasmithImporter/Source/DatasmithImporter/Private/DatasmithFileProducer.cpp
@@ -14,8 +14,8 @@
 #include "DatasmithScene.h"
 #include "DatasmithSceneActor.h"
 #include "DatasmithSceneFactory.h"
+#include "DatasmithTranslatorManager.h"
 #include "IDatasmithSceneElements.h"
-#include "Translators/DatasmithTranslatorManager.h"
 #include "Utility/DatasmithImporterUtils.h"
 
 #include "Async/ParallelFor.h"
@@ -172,7 +172,8 @@ bool UDatasmithFileProducer::Initialize()
 	TSharedRef< IDatasmithScene > SceneElement = FDatasmithSceneFactory::CreateScene( *Source.GetSceneName() );
 
 	constexpr EObjectFlags LocalObjectFlags = RF_Public | RF_Standalone | RF_Transactional;
-	if ( !ImportContextPtr->Init( SceneElement, RootPath, LocalObjectFlags, Context.ProgressReporterPtr->GetFeedbackContext(), TSharedPtr< FJsonObject >(), true ) )
+	FFeedbackContext* InFeedbackContext = Context.ProgressReporterPtr ? Context.ProgressReporterPtr->GetFeedbackContext() : nullptr;
+	if ( !ImportContextPtr->Init( SceneElement, RootPath, LocalObjectFlags, InFeedbackContext, TSharedPtr< FJsonObject >(), true ) )
 	{
 		LogError( LOCTEXT( "DatasmithFileProducer_Initialization", "Initialization of producer failed." ) );
 		return false;
diff --git a/Engine/Plugins/Enterprise/DatasmithImporter/Source/DatasmithImporter/Private/DatasmithFileProducer.h b/Engine/Plugins/Enterprise/DatasmithImporter/Source/DatasmithImporter/Private/DatasmithFileProducer.h
--- a/Engine/Plugins/Enterprise/DatasmithImporter/Source/DatasmithImporter/Private/DatasmithFileProducer.h
+++ b/Engine/Plugins/Enterprise/DatasmithImporter/Source/DatasmithImporter/Private/DatasmithFileProducer.h
@@ -10,7 +10,7 @@
 #include "DatasmithImportContext.h"
 #include "DatasmithImportOptions.h"
 #include "DatasmithScene.h"
-#include "Translators/DatasmithTranslatableSource.h"
+#include "DatasmithTranslatableSource.h"
 
 #include "IDetailCustomization.h"
 #include "Input/Reply.h"
diff --git a/Engine/Plugins/Enterprise/DatasmithImporter/Source/DatasmithImporter/Private/DatasmithImportContext.cpp b/Engine/Plugins/Enterprise/DatasmithImporter/Source/DatasmithImporter/Private/DatasmithImportContext.cpp
--- a/Engine/Plugins/Enterprise/DatasmithImporter/Source/DatasmithImporter/Private/DatasmithImportContext.cpp
+++ b/Engine/Plugins/Enterprise/DatasmithImporter/Source/DatasmithImporter/Private/DatasmithImportContext.cpp
@@ -5,10 +5,10 @@
 #include "DatasmithActorImporter.h"
 #include "DatasmithImportOptions.h"
 #include "DatasmithSceneFactory.h"
+#include "DatasmithTranslator.h"
 #include "DatasmithUtils.h"
 #include "UI/DatasmithImportOptionsWindow.h"
 #include "Utility/DatasmithImporterUtils.h"
-#include "Translators/DatasmithTranslator.h"
 
 #include "Misc/FeedbackContext.h"
 #include "PackageTools.h"
diff --git a/Engine/Plugins/Enterprise/DatasmithImporter/Source/DatasmithImporter/Private/DatasmithImportFactory.cpp b/Engine/Plugins/Enterprise/DatasmithImporter/Source/DatasmithImporter/Private/DatasmithImportFactory.cpp
--- a/Engine/Plugins/Enterprise/DatasmithImporter/Source/DatasmithImporter/Private/DatasmithImportFactory.cpp
+++ b/Engine/Plugins/Enterprise/DatasmithImporter/Source/DatasmithImporter/Private/DatasmithImportFactory.cpp
@@ -10,10 +10,10 @@
 #include "DatasmithSceneActor.h"
 #include "DatasmithSceneFactory.h"
 #include "DatasmithStaticMeshImporter.h"
+#include "DatasmithTranslatableSource.h"
+#include "DatasmithTranslatorManager.h"
 #include "IDatasmithSceneElements.h"
 #include "LayoutUV.h"
-#include "Translators/DatasmithTranslatableSource.h"
-#include "Translators/DatasmithTranslatorManager.h"
 #include "Utility/DatasmithImporterUtils.h"
 #include "Utility/DatasmithImportFactoryHelper.h"
 
diff --git a/Engine/Plugins/Enterprise/DatasmithImporter/Source/DatasmithImporter/Private/DatasmithImporter.cpp b/Engine/Plugins/Enterprise/DatasmithImporter/Source/DatasmithImporter/Private/DatasmithImporter.cpp
--- a/Engine/Plugins/Enterprise/DatasmithImporter/Source/DatasmithImporter/Private/DatasmithImporter.cpp
+++ b/Engine/Plugins/Enterprise/DatasmithImporter/Source/DatasmithImporter/Private/DatasmithImporter.cpp
@@ -12,18 +12,18 @@
 #include "DatasmithLevelVariantSetsImporter.h"
 #include "DatasmithLightImporter.h"
 #include "DatasmithMaterialImporter.h"
+#include "DatasmithPayload.h"
 #include "DatasmithPostProcessImporter.h"
 #include "DatasmithScene.h"
 #include "DatasmithSceneActor.h"
 #include "DatasmithSceneFactory.h"
 #include "DatasmithStaticMeshImporter.h"
+#include "DatasmithTranslator.h"
 #include "DatasmithTextureImporter.h"
 #include "IDatasmithSceneElements.h"
 #include "DatasmithAnimationElements.h"
 #include "LevelVariantSets.h"
 #include "ObjectTemplates/DatasmithObjectTemplate.h"
-#include "Translators/DatasmithPayload.h"
-#include "Translators/DatasmithTranslator.h"
 #include "Utility/DatasmithImporterUtils.h"
 #include "Utility/DatasmithTextureResize.h"
 
@@ -1354,7 +1354,7 @@ UMaterialInterface* FDatasmithImporter::ImportMaterial( FDatasmithImportContext&
 	return ImportedMaterial;
 }
 
-UObject* FDatasmithImporter::FinalizeMaterial( UObject* SourceMaterial, const TCHAR* MaterialsFolderPath, UMaterialInterface* ExistingMaterial, TMap< UObject*, UObject* >* ReferencesToRemap )
+UObject* FDatasmithImporter::FinalizeMaterial( UObject* SourceMaterial, const TCHAR* MaterialsFolderPath, UMaterialInterface* ExistingMaterial, TMap< UObject*, UObject* >* ReferencesToRemap)
 {
 	TRACE_CPUPROFILER_EVENT_SCOPE(FDatasmithImporter::FinalizeMaterial);
 
@@ -1372,18 +1372,22 @@ UObject* FDatasmithImporter::FinalizeMaterial( UObject* SourceMaterial, const TC
 	{
 		if ( UMaterialInterface* SourceMaterialParent = SourceMaterialInstance->Parent )
 		{
-			FString SourceMaterialPath = SourceMaterialInstance->GetOutermost()->GetName();
-			FString SourceParentPath = SourceMaterialParent->GetOutermost()->GetName();
-
-			if ( SourceParentPath.StartsWith( SourceMaterialPath ) )
+			// Do not finalize parent material more than once by verifying it is not already in ReferencesToRemap
+			if (!ReferencesToRemap->Contains(SourceMaterialParent))
 			{
-				// Simply finalize the source parent material.
-				// Note that the parent material will be overridden on the existing material instance
-				FString DestinationParentPath = SourceParentPath;
-				DestinationParentPath.RemoveFromStart( SourceMaterialPath );
-				DestinationParentPath = MaterialsFolderPath / DestinationParentPath;
+				FString SourceMaterialPath = SourceMaterialInstance->GetOutermost()->GetName();
+				FString SourceParentPath = SourceMaterialParent->GetOutermost()->GetName();
+
+				if (SourceParentPath.StartsWith( SourceMaterialPath ) )
+				{
+					// Simply finalize the source parent material.
+					// Note that the parent material will be overridden on the existing material instance
+					FString DestinationParentPath = SourceParentPath;
+					DestinationParentPath.RemoveFromStart( SourceMaterialPath );
+					DestinationParentPath = MaterialsFolderPath / DestinationParentPath;
 
-				FinalizeMaterial( SourceMaterialParent, *DestinationParentPath, nullptr, ReferencesToRemap );
+					FinalizeMaterial( SourceMaterialParent, *DestinationParentPath, nullptr, ReferencesToRemap );
+				}
 			}
 		}
 	}
diff --git a/Engine/Plugins/Enterprise/DatasmithImporter/Source/DatasmithImporter/Private/DatasmithImporterModule.cpp b/Engine/Plugins/Enterprise/DatasmithImporter/Source/DatasmithImporter/Private/DatasmithImporterModule.cpp
--- a/Engine/Plugins/Enterprise/DatasmithImporter/Source/DatasmithImporter/Private/DatasmithImporterModule.cpp
+++ b/Engine/Plugins/Enterprise/DatasmithImporter/Source/DatasmithImporter/Private/DatasmithImporterModule.cpp
@@ -16,11 +16,6 @@
 #include "DatasmithScene.h"
 #include "DatasmithStaticMeshImporter.h"
 #include "DatasmithUtils.h"
-#include "MasterMaterials/DatasmithC4DMaterialSelector.h"
-#include "MasterMaterials/DatasmithCityEngineMaterialSelector.h"
-#include "MasterMaterials/DatasmithMasterMaterialManager.h"
-#include "MasterMaterials/DatasmithRevitMaterialSelector.h"
-#include "MasterMaterials/DatasmithSketchupMaterialSelector.h"
 #include "ObjectTemplates/DatasmithObjectTemplate.h"
 #include "ObjectTemplates/DatasmithStaticMeshTemplate.h"
 #include "UI/DatasmithUIManager.h"
@@ -31,7 +26,7 @@
 #include "ContentBrowserModule.h"
 #include "DataprepAssetInterface.h"
 #include "DataprepAssetUserData.h"
-#include "DataprepCoreLibrary.h"
+#include "DataprepCoreUtils.h"
 #include "Editor.h"
 #include "EditorFramework/AssetImportData.h"
 #include "EditorStyleSet.h"
@@ -64,13 +59,6 @@ class FDatasmithImporterModule : public IDatasmithImporterModule
 public:
 	virtual void StartupModule() override
 	{
-		FDatasmithMasterMaterialManager::Create();
-
-		FDatasmithMasterMaterialManager::Get().RegisterSelector(TEXT("C4D"), MakeShared< FDatasmithC4DMaterialSelector >());
-		FDatasmithMasterMaterialManager::Get().RegisterSelector(TEXT("Revit"), MakeShared< FDatasmithRevitMaterialSelector >());
-		FDatasmithMasterMaterialManager::Get().RegisterSelector( TEXT("SketchUp"), MakeShared< FDatasmithSketchUpMaterialSelector >() );
-		FDatasmithMasterMaterialManager::Get().RegisterSelector( TEXT("CityEngine"), MakeShared< FDatasmithCityEngineMaterialSelector >() );
-
 		UDatasmithFileProducer::LoadDefaultSettings();
 
 		// Disable any UI feature if running in command mode
@@ -119,10 +107,6 @@ public:
 			FPropertyEditorModule& PropertyModule = FModuleManager::LoadModuleChecked< FPropertyEditorModule >( TEXT("PropertyEditor") );
 			PropertyModule.UnregisterCustomClassLayout( TEXT("DatasmithFileProducer") );
 		}
-
-		FDatasmithMasterMaterialManager::Get().UnregisterSelector( TEXT("CityEngine") );
-
-		FDatasmithMasterMaterialManager::Destroy();
 	}
 
 	virtual void ResetOverrides( UObject* Object ) override
@@ -252,7 +236,9 @@ void FDatasmithImporterModule::AddDataprepMenuEntryForDatasmithSceneAsset()
 						{
 							for ( const TStrongObjectPtr<UDataprepAssetInterface>& DataprepAssetInterfacePtr : DataprepAssetInterfacesPtr )
 							{
-								UDataprepCoreLibrary::ExecuteWithReporting( DataprepAssetInterfacePtr.Get() );
+								FDataprepCoreUtils::ExecuteDataprep( DataprepAssetInterfacePtr.Get()
+									, MakeShared<FDataprepCoreUtils::FDataprepLogger>()
+									, MakeShared<FDataprepCoreUtils::FDataprepProgressUIReporter>() );
 							}
 						});
 
diff --git a/Engine/Plugins/Enterprise/DatasmithImporter/Source/DatasmithImporter/Private/DatasmithStaticMeshImporter.cpp b/Engine/Plugins/Enterprise/DatasmithImporter/Source/DatasmithImporter/Private/DatasmithStaticMeshImporter.cpp
--- a/Engine/Plugins/Enterprise/DatasmithImporter/Source/DatasmithImporter/Private/DatasmithStaticMeshImporter.cpp
+++ b/Engine/Plugins/Enterprise/DatasmithImporter/Source/DatasmithImporter/Private/DatasmithStaticMeshImporter.cpp
@@ -5,12 +5,12 @@
 #include "DatasmithImportContext.h"
 #include "DatasmithImporterModule.h"
 #include "DatasmithMaterialImporter.h"
-#include "DatasmithMeshHelper.h"
 #include "DatasmithMeshUObject.h"
 #include "IDatasmithSceneElements.h"
 #include "ObjectTemplates/DatasmithStaticMeshTemplate.h"
-#include "Translators/DatasmithPayload.h"
+#include "DatasmithPayload.h"
 #include "Utility/DatasmithImporterUtils.h"
+#include "Utility/DatasmithMeshHelper.h"
 
 #include "Algo/AnyOf.h"
 #include "Async/Async.h"
diff --git a/Engine/Plugins/Enterprise/DatasmithImporter/Source/DatasmithImporter/Public/DatasmithBlueprintLibrary.h b/Engine/Plugins/Enterprise/DatasmithImporter/Source/DatasmithImporter/Public/DatasmithBlueprintLibrary.h
--- a/Engine/Plugins/Enterprise/DatasmithImporter/Source/DatasmithImporter/Public/DatasmithBlueprintLibrary.h
+++ b/Engine/Plugins/Enterprise/DatasmithImporter/Source/DatasmithImporter/Public/DatasmithBlueprintLibrary.h
@@ -3,8 +3,8 @@
 #pragma once
 
 #include "DatasmithImportContext.h"
+#include "DatasmithTranslatableSource.h"
 #include "ObjectElements/DatasmithUSceneElement.h"
-#include "Translators/DatasmithTranslatableSource.h"
 
 #include "Kismet/BlueprintFunctionLibrary.h"
 
diff --git a/Engine/Plugins/Enterprise/DatasmithImporter/Source/DatasmithNativeTranslator/Private/DatasmithNativeTranslator.cpp b/Engine/Plugins/Enterprise/DatasmithImporter/Source/DatasmithNativeTranslator/Private/DatasmithNativeTranslator.cpp
--- a/Engine/Plugins/Enterprise/DatasmithImporter/Source/DatasmithNativeTranslator/Private/DatasmithNativeTranslator.cpp
+++ b/Engine/Plugins/Enterprise/DatasmithImporter/Source/DatasmithNativeTranslator/Private/DatasmithNativeTranslator.cpp
@@ -4,11 +4,11 @@
 
 #include "DatasmithAnimationElements.h"
 #include "DatasmithAnimationSerializer.h"
-#include "DatasmithMeshHelper.h"
 #include "DatasmithMeshUObject.h"
 #include "DatasmithSceneSource.h"
 #include "DatasmithSceneXmlReader.h"
 #include "IDatasmithSceneElements.h"
+#include "Utility/DatasmithMeshHelper.h"
 
 #include "HAL/FileManager.h"
 #include "MeshDescriptionOperations.h"
@@ -18,7 +18,6 @@
 #include "UObject/StrongObjectPtr.h"
 
 
-
 void FDatasmithNativeTranslator::Initialize(FDatasmithTranslatorCapabilities& OutCapabilities)
 {
 	OutCapabilities.SupportedFileFormats.Emplace(TEXT("udatasmith"), TEXT("Datasmith files"));
diff --git a/Engine/Plugins/Enterprise/DatasmithImporter/Source/DatasmithTranslator/Private/DatasmithPayload.cpp b/Engine/Plugins/Enterprise/DatasmithImporter/Source/DatasmithTranslator/Private/DatasmithPayload.cpp
new file mode 100644
--- /dev/null
+++ b/Engine/Plugins/Enterprise/DatasmithImporter/Source/DatasmithTranslator/Private/DatasmithPayload.cpp
@@ -0,0 +1,3 @@
+// Copyright 1998-2019 Epic Games, Inc. All Rights Reserved.
+
+#include "DatasmithPayload.h"
diff --git a/Engine/Plugins/Enterprise/DatasmithImporter/Source/DatasmithTranslator/Private/DatasmithSceneSource.cpp b/Engine/Plugins/Enterprise/DatasmithImporter/Source/DatasmithTranslator/Private/DatasmithSceneSource.cpp
new file mode 100644
--- /dev/null
+++ b/Engine/Plugins/Enterprise/DatasmithImporter/Source/DatasmithTranslator/Private/DatasmithSceneSource.cpp
@@ -0,0 +1,34 @@
+// Copyright 1998-2019 Epic Games, Inc. All Rights Reserved.
+
+#include "DatasmithSceneSource.h"
+
+#include "DatasmithTranslator.h"
+#include "DatasmithTranslatorManager.h"
+
+#include "Misc/Paths.h"
+
+#define LOCTEXT_NAMESPACE "DatasmithSceneSource"
+
+void FDatasmithSceneSource::SetSourceFile(const FString& InFilePath)
+{
+	FilePath = InFilePath;
+	SceneDeducedName = FPaths::GetBaseFilename(FilePath);
+	FileExtension = FPaths::GetExtension(FilePath);
+	if (FCString::IsNumeric(*FileExtension))
+	{
+		FileExtension = FPaths::GetExtension(SceneDeducedName) + TEXT(".") + FileExtension;
+		SceneDeducedName = FPaths::GetBaseFilename(SceneDeducedName);
+	}
+}
+
+void FDatasmithSceneSource::SetSceneName(const FString& InSceneName)
+{
+	SceneOverrideName = InSceneName;
+}
+
+const FString& FDatasmithSceneSource::GetSceneName() const
+{
+	return SceneOverrideName.IsEmpty() ? SceneDeducedName : SceneOverrideName;
+}
+
+#undef LOCTEXT_NAMESPACE
diff --git a/Engine/Plugins/Enterprise/DatasmithImporter/Source/DatasmithTranslator/Private/DatasmithTranslatorModule.cpp b/Engine/Plugins/Enterprise/DatasmithImporter/Source/DatasmithTranslator/Private/DatasmithTranslatorModule.cpp
new file mode 100644
--- /dev/null
+++ b/Engine/Plugins/Enterprise/DatasmithImporter/Source/DatasmithTranslator/Private/DatasmithTranslatorModule.cpp
@@ -0,0 +1,32 @@
+// Copyright 1998-2019 Epic Games, Inc. All Rights Reserved.
+
+#include "DatasmithTranslatorModule.h"
+
+#include "MasterMaterials/DatasmithC4DMaterialSelector.h"
+#include "MasterMaterials/DatasmithCityEngineMaterialSelector.h"
+#include "MasterMaterials/DatasmithMasterMaterialManager.h"
+#include "MasterMaterials/DatasmithRevitMaterialSelector.h"
+#include "MasterMaterials/DatasmithSketchupMaterialSelector.h"
+
+void IDatasmithTranslatorModule::StartupModule()
+{
+	FDatasmithMasterMaterialManager::Create();
+
+	//A minimal set of natively supported master materials
+	FDatasmithMasterMaterialManager::Get().RegisterSelector(TEXT("C4D"), MakeShared< FDatasmithC4DMaterialSelector >());
+	FDatasmithMasterMaterialManager::Get().RegisterSelector(TEXT("Revit"), MakeShared< FDatasmithRevitMaterialSelector >());
+	FDatasmithMasterMaterialManager::Get().RegisterSelector(TEXT("SketchUp"), MakeShared< FDatasmithSketchUpMaterialSelector >());
+	FDatasmithMasterMaterialManager::Get().RegisterSelector(TEXT("CityEngine"), MakeShared< FDatasmithCityEngineMaterialSelector >());
+}
+
+void IDatasmithTranslatorModule::ShutdownModule()
+{
+	FDatasmithMasterMaterialManager::Get().UnregisterSelector(TEXT("C4D"));
+	FDatasmithMasterMaterialManager::Get().UnregisterSelector(TEXT("Revit"));
+	FDatasmithMasterMaterialManager::Get().UnregisterSelector(TEXT("SketchUp"));
+	FDatasmithMasterMaterialManager::Get().UnregisterSelector(TEXT("CityEngine"));
+
+	FDatasmithMasterMaterialManager::Destroy();
+}
+
+IMPLEMENT_MODULE(IDatasmithTranslatorModule, DatasmithTranslator);
\ No newline at end of file
diff --git a/Engine/Plugins/Enterprise/DatasmithImporter/Source/DatasmithTranslator/Private/MasterMaterials/DatasmithC4DMaterialSelector.cpp b/Engine/Plugins/Enterprise/DatasmithImporter/Source/DatasmithTranslator/Private/MasterMaterials/DatasmithC4DMaterialSelector.cpp
new file mode 100644
--- /dev/null
+++ b/Engine/Plugins/Enterprise/DatasmithImporter/Source/DatasmithTranslator/Private/MasterMaterials/DatasmithC4DMaterialSelector.cpp
@@ -0,0 +1,46 @@
+// Copyright 1998-2019 Epic Games, Inc. All Rights Reserved.
+
+#include "DatasmithC4DMaterialSelector.h"
+
+#include "IDatasmithSceneElements.h"
+
+#include "Materials/Material.h"
+#include "Materials/MaterialInstanceConstant.h"
+#include "UObject/SoftObjectPath.h"
+#include "Templates/Casts.h"
+
+FDatasmithC4DMaterialSelector::FDatasmithC4DMaterialSelector()
+{
+	MasterMaterial.FromSoftObjectPath( FSoftObjectPath("/DatasmithContent/Materials/C4DMaster.C4DMaster") );
+}
+
+bool FDatasmithC4DMaterialSelector::IsValid() const
+{
+	return MasterMaterial.IsValid();
+}
+
+const FDatasmithMasterMaterial& FDatasmithC4DMaterialSelector::GetMasterMaterial( const TSharedPtr< IDatasmithMasterMaterialElement >& /*InDatasmithMaterial*/ ) const
+{
+	return MasterMaterial;
+}
+
+void FDatasmithC4DMaterialSelector::FinalizeMaterialInstance(const TSharedPtr<IDatasmithMasterMaterialElement>& InDatasmithMaterial, UMaterialInstanceConstant * MaterialInstance) const
+{
+	if (!InDatasmithMaterial.IsValid() || MaterialInstance == nullptr)
+	{
+		return;
+	}
+
+	// Set blend mode to translucent if material requires transparency.
+	if (InDatasmithMaterial->GetMaterialType() == EDatasmithMasterMaterialType::Transparent)
+	{
+		MaterialInstance->BasePropertyOverrides.bOverride_BlendMode = true;
+		MaterialInstance->BasePropertyOverrides.BlendMode = BLEND_Translucent;
+	}
+	// Set blend mode to masked if material has cutouts.
+	else if (InDatasmithMaterial->GetMaterialType() == EDatasmithMasterMaterialType::CutOut)
+	{
+		MaterialInstance->BasePropertyOverrides.bOverride_BlendMode = true;
+		MaterialInstance->BasePropertyOverrides.BlendMode = BLEND_Masked;
+	}
+}
diff --git a/Engine/Plugins/Enterprise/DatasmithImporter/Source/DatasmithTranslator/Private/MasterMaterials/DatasmithC4DMaterialSelector.h b/Engine/Plugins/Enterprise/DatasmithImporter/Source/DatasmithTranslator/Private/MasterMaterials/DatasmithC4DMaterialSelector.h
new file mode 100644
--- /dev/null
+++ b/Engine/Plugins/Enterprise/DatasmithImporter/Source/DatasmithTranslator/Private/MasterMaterials/DatasmithC4DMaterialSelector.h
@@ -0,0 +1,27 @@
+// Copyright 1998-2019 Epic Games, Inc. All Rights Reserved.
+
+#pragma once
+
+#include "DatasmithDefinitions.h"
+#include "MasterMaterials/DatasmithMasterMaterial.h"
+#include "MasterMaterials/DatasmithMasterMaterialSelector.h"
+
+#include "Templates/SharedPointer.h"
+
+class IDatasmithMasterMaterialElement;
+
+class FDatasmithC4DMaterialSelector : public FDatasmithMasterMaterialSelector
+{
+public:
+
+	FDatasmithC4DMaterialSelector();
+	virtual ~FDatasmithC4DMaterialSelector() = default;
+
+	virtual bool IsValid() const override;
+	virtual const FDatasmithMasterMaterial& GetMasterMaterial( const TSharedPtr< IDatasmithMasterMaterialElement >& /*InDatasmithMaterial*/ ) const override;
+	virtual void FinalizeMaterialInstance(const TSharedPtr< IDatasmithMasterMaterialElement >& InDatasmithMaterial, UMaterialInstanceConstant* MaterialInstance) const override;
+
+private:
+
+	FDatasmithMasterMaterial MasterMaterial;
+};
diff --git a/Engine/Plugins/Enterprise/DatasmithImporter/Source/DatasmithTranslator/Private/MasterMaterials/DatasmithCityEngineMaterialSelector.cpp b/Engine/Plugins/Enterprise/DatasmithImporter/Source/DatasmithTranslator/Private/MasterMaterials/DatasmithCityEngineMaterialSelector.cpp
new file mode 100644
--- /dev/null
+++ b/Engine/Plugins/Enterprise/DatasmithImporter/Source/DatasmithTranslator/Private/MasterMaterials/DatasmithCityEngineMaterialSelector.cpp
@@ -0,0 +1,65 @@
+// Copyright 1998-2019 Epic Games, Inc. All Rights Reserved.
+
+#include "DatasmithCityEngineMaterialSelector.h"
+
+#include "IDatasmithSceneElements.h"
+
+#include "Materials/Material.h"
+#include "UObject/SoftObjectPath.h"
+#include "Templates/Casts.h"
+
+FDatasmithCityEngineMaterialSelector::FDatasmithCityEngineMaterialSelector()
+{
+	// Master
+	MasterMaterial.FromSoftObjectPath( FSoftObjectPath("/DatasmithContent/Materials/CE_OpaqueMaster.CE_OpaqueMaster") );
+	MasterMaterialTransparent.FromSoftObjectPath( FSoftObjectPath("/DatasmithContent/Materials/CE_OpacityMaster.CE_OpacityMaster") );
+	MasterMaterialTransparentSimple.FromSoftObjectPath( FSoftObjectPath("/DatasmithContent/Materials/CE_OpacitySimpleMaster.CE_OpacitySimpleMaster") );
+}
+
+bool FDatasmithCityEngineMaterialSelector::IsValid() const
+{
+	return MasterMaterial.IsValid() && MasterMaterialTransparent.IsValid() && MasterMaterialTransparentSimple.IsValid();
+}
+
+const FDatasmithMasterMaterial& FDatasmithCityEngineMaterialSelector::GetMasterMaterial( const TSharedPtr< IDatasmithMasterMaterialElement >& InDatasmithMaterial ) const
+{
+	EDatasmithMasterMaterialType MaterialType = InDatasmithMaterial->GetMaterialType();
+
+	if ( !IsValidMaterialType( MaterialType ) )
+	{
+		MaterialType = EDatasmithMasterMaterialType::Auto;
+	}
+
+	bool bIsTransparent = MaterialType == EDatasmithMasterMaterialType::Transparent;
+
+	if ( MaterialType == EDatasmithMasterMaterialType::Auto )
+	{
+		const TSharedPtr< IDatasmithKeyValueProperty > OpacityProperty = InDatasmithMaterial->GetPropertyByName(TEXT("Opacity"));
+		const TSharedPtr< IDatasmithKeyValueProperty > OpacityMapProperty = InDatasmithMaterial->GetPropertyByName(TEXT("OpacityMap"));
+
+		float OpacityValue;
+		bIsTransparent = ( OpacityProperty.IsValid() && GetFloat( OpacityProperty, OpacityValue ) && OpacityValue < 1.f );
+
+		FString OpacityMap;
+		bIsTransparent = bIsTransparent || ( OpacityMapProperty.IsValid() && GetTexture( OpacityMapProperty, OpacityMap ) && !OpacityMap.IsEmpty() );
+	}
+
+	if ( bIsTransparent )
+	{
+		if ( InDatasmithMaterial->GetQuality() == EDatasmithMasterMaterialQuality::Low )
+		{
+			return MasterMaterialTransparentSimple;
+		}
+		else
+		{
+			return MasterMaterialTransparent;
+		}
+	}
+
+	return MasterMaterial;
+}
+
+bool FDatasmithCityEngineMaterialSelector::IsValidMaterialType( EDatasmithMasterMaterialType InType ) const
+{
+	return InType == EDatasmithMasterMaterialType::Auto || InType == EDatasmithMasterMaterialType::Opaque || InType == EDatasmithMasterMaterialType::Transparent;
+}
diff --git a/Engine/Plugins/Enterprise/DatasmithImporter/Source/DatasmithTranslator/Private/MasterMaterials/DatasmithCityEngineMaterialSelector.h b/Engine/Plugins/Enterprise/DatasmithImporter/Source/DatasmithTranslator/Private/MasterMaterials/DatasmithCityEngineMaterialSelector.h
new file mode 100644
--- /dev/null
+++ b/Engine/Plugins/Enterprise/DatasmithImporter/Source/DatasmithTranslator/Private/MasterMaterials/DatasmithCityEngineMaterialSelector.h
@@ -0,0 +1,27 @@
+// Copyright 1998-2019 Epic Games, Inc. All Rights Reserved.
+#pragma once
+
+#include "DatasmithDefinitions.h"
+#include "MasterMaterials/DatasmithMasterMaterial.h"
+#include "MasterMaterials/DatasmithMasterMaterialSelector.h"
+
+#include "Templates/SharedPointer.h"
+
+class IDatasmithMasterMaterialElement;
+
+class FDatasmithCityEngineMaterialSelector : public FDatasmithMasterMaterialSelector
+{
+public:
+	FDatasmithCityEngineMaterialSelector();
+
+	virtual bool IsValid() const override;
+	virtual const FDatasmithMasterMaterial& GetMasterMaterial( const TSharedPtr< IDatasmithMasterMaterialElement >& InDatasmithMaterial ) const override;
+
+protected:
+	bool IsValidMaterialType( EDatasmithMasterMaterialType InType ) const;
+
+private:
+	FDatasmithMasterMaterial MasterMaterial;
+	FDatasmithMasterMaterial MasterMaterialTransparent;
+	FDatasmithMasterMaterial MasterMaterialTransparentSimple;
+};
diff --git a/Engine/Plugins/Enterprise/DatasmithImporter/Source/DatasmithTranslator/Private/MasterMaterials/DatasmithMasterMaterial.cpp b/Engine/Plugins/Enterprise/DatasmithImporter/Source/DatasmithTranslator/Private/MasterMaterials/DatasmithMasterMaterial.cpp
new file mode 100644
--- /dev/null
+++ b/Engine/Plugins/Enterprise/DatasmithImporter/Source/DatasmithTranslator/Private/MasterMaterials/DatasmithMasterMaterial.cpp
@@ -0,0 +1,53 @@
+// Copyright 1998-2019 Epic Games, Inc. All Rights Reserved.
+
+#include "MasterMaterials/DatasmithMasterMaterial.h"
+
+#include "Materials/Material.h"
+#include "Materials/MaterialExpressionScalarParameter.h"
+#include "Materials/MaterialExpressionStaticBoolParameter.h"
+#include "Materials/MaterialExpressionTextureSampleParameter.h"
+#include "Materials/MaterialExpressionVectorParameter.h"
+
+#include "UObject/SoftObjectPath.h"
+
+FDatasmithMasterMaterial::FDatasmithMasterMaterial()
+	: Material( nullptr )
+{
+}
+
+void FDatasmithMasterMaterial::FromMaterial( UMaterial* InMaterial )
+{
+#if WITH_EDITOR
+	if ( InMaterial )
+	{
+		for ( UMaterialExpression* Expression : InMaterial->Expressions )
+		{
+			FString ExpressionName = Expression->GetName();
+
+			if ( Expression->IsA< UMaterialExpressionVectorParameter >() )
+			{
+				VectorParams.Add( Expression->GetParameterName().ToString() );
+			}
+			else if ( Expression->IsA< UMaterialExpressionScalarParameter >() )
+			{
+				ScalarParams.Add( Expression->GetParameterName().ToString() );
+			}
+			else if ( Expression->IsA< UMaterialExpressionTextureSampleParameter >() )
+			{
+				TextureParams.Add( Expression->GetParameterName().ToString() );
+			}
+			else if ( Expression->IsA< UMaterialExpressionStaticBoolParameter >() )
+			{
+				BoolParams.Add( Expression->GetParameterName().ToString() );
+			}
+		}
+	}
+#endif
+
+	Material = InMaterial;
+}
+
+void FDatasmithMasterMaterial::FromSoftObjectPath( const FSoftObjectPath& InObjectPath)
+{
+	FromMaterial( Cast< UMaterial >(InObjectPath.TryLoad() ) );
+}
diff --git a/Engine/Plugins/Enterprise/DatasmithImporter/Source/DatasmithTranslator/Private/MasterMaterials/DatasmithMasterMaterialManager.cpp b/Engine/Plugins/Enterprise/DatasmithImporter/Source/DatasmithTranslator/Private/MasterMaterials/DatasmithMasterMaterialManager.cpp
new file mode 100644
--- /dev/null
+++ b/Engine/Plugins/Enterprise/DatasmithImporter/Source/DatasmithTranslator/Private/MasterMaterials/DatasmithMasterMaterialManager.cpp
@@ -0,0 +1,65 @@
+// Copyright 1998-2019 Epic Games, Inc. All Rights Reserved.
+
+#include "MasterMaterials/DatasmithMasterMaterialManager.h"
+
+#include "MasterMaterials/DatasmithMasterMaterialSelector.h"
+
+TUniquePtr< FDatasmithMasterMaterialManager > FDatasmithMasterMaterialManager::Instance;
+
+void FDatasmithMasterMaterialManager::Create()
+{
+	Instance = MakeUnique< FDatasmithMasterMaterialManager >();
+}
+
+void FDatasmithMasterMaterialManager::Destroy()
+{
+	Instance.Reset();
+}
+
+FDatasmithMasterMaterialManager& FDatasmithMasterMaterialManager::Get()
+{
+	check( Instance.IsValid() );
+	return *Instance.Get();
+}
+
+FString FDatasmithMasterMaterialManager::GetHostFromString( const FString& HostString )
+{
+	if ( HostString.Contains( TEXT("CityEngine") ) )
+	{
+		return TEXT("CityEngine");
+	}
+	else if ( HostString.Contains( TEXT("Deltagen") ) )
+	{
+		return TEXT("Deltagen");
+	}
+	else if ( HostString.Contains(TEXT("VRED") ) )
+	{
+		return TEXT("VRED");
+	}
+	else
+	{
+		return HostString;
+	}
+}
+
+void FDatasmithMasterMaterialManager::RegisterSelector( const TCHAR* Host, TSharedPtr< FDatasmithMasterMaterialSelector > Selector )
+{
+	Selectors.FindOrAdd( Host ) = Selector;
+}
+
+void FDatasmithMasterMaterialManager::UnregisterSelector( const TCHAR* Host )
+{
+	Selectors.Remove( Host );
+}
+
+const TSharedPtr< FDatasmithMasterMaterialSelector > FDatasmithMasterMaterialManager::GetSelector( const TCHAR* Host ) const
+{
+	if ( Selectors.Contains( Host ) )
+	{
+		return Selectors[ Host ];
+	}
+	else
+	{
+		return MakeShared< FDatasmithMasterMaterialSelector >();
+	}
+}
diff --git a/Engine/Plugins/Enterprise/DatasmithImporter/Source/DatasmithTranslator/Private/MasterMaterials/DatasmithMasterMaterialSelector.cpp b/Engine/Plugins/Enterprise/DatasmithImporter/Source/DatasmithTranslator/Private/MasterMaterials/DatasmithMasterMaterialSelector.cpp
new file mode 100644
--- /dev/null
+++ b/Engine/Plugins/Enterprise/DatasmithImporter/Source/DatasmithTranslator/Private/MasterMaterials/DatasmithMasterMaterialSelector.cpp
@@ -0,0 +1,71 @@
+// Copyright 1998-2019 Epic Games, Inc. All Rights Reserved.
+
+#include "MasterMaterials/DatasmithMasterMaterialSelector.h"
+
+#include "MasterMaterials/DatasmithMasterMaterial.h"
+#include "IDatasmithSceneElements.h"
+
+FDatasmithMasterMaterial FDatasmithMasterMaterialSelector::InvalidMasterMaterial;
+
+const FDatasmithMasterMaterial& FDatasmithMasterMaterialSelector::GetMasterMaterial( const TSharedPtr< IDatasmithMasterMaterialElement >& InDatasmithMaterial ) const
+{
+	return InvalidMasterMaterial;
+}
+
+bool FDatasmithMasterMaterialSelector::GetColor( const TSharedPtr< IDatasmithKeyValueProperty >& InMaterialProperty, FLinearColor& OutValue ) const
+{
+	if (!InMaterialProperty.IsValid() || InMaterialProperty->GetPropertyType() != EDatasmithKeyValuePropertyType::Color )
+	{
+		return false;
+	}
+
+	return OutValue.InitFromString( InMaterialProperty->GetValue() ); // TODO: Handle sRGB vs linear RGB properly?
+}
+
+bool FDatasmithMasterMaterialSelector::GetFloat( const TSharedPtr< IDatasmithKeyValueProperty >& InMaterialProperty, float& OutValue ) const
+{
+	if (!InMaterialProperty.IsValid() || InMaterialProperty->GetPropertyType() != EDatasmithKeyValuePropertyType::Float )
+	{
+		return false;
+	}
+
+	OutValue = FCString::Atof( InMaterialProperty->GetValue() );
+
+	return true;
+}
+
+bool FDatasmithMasterMaterialSelector::GetBool( const TSharedPtr< IDatasmithKeyValueProperty >& InMaterialProperty, bool& OutValue ) const
+{
+	if (!InMaterialProperty.IsValid() || InMaterialProperty->GetPropertyType() != EDatasmithKeyValuePropertyType::Bool )
+	{
+		return false;
+	}
+
+	OutValue = FString( InMaterialProperty->GetValue() ).ToBool();
+
+	return true;
+}
+
+bool FDatasmithMasterMaterialSelector::GetTexture( const TSharedPtr< IDatasmithKeyValueProperty >& InMaterialProperty, FString& OutValue ) const
+{
+	if (!InMaterialProperty.IsValid() || InMaterialProperty->GetPropertyType() != EDatasmithKeyValuePropertyType::Texture )
+	{
+		return false;
+	}
+
+	OutValue = InMaterialProperty->GetValue();
+
+	return true;
+}
+
+bool FDatasmithMasterMaterialSelector::GetString(const TSharedPtr< IDatasmithKeyValueProperty >& InMaterialProperty, FString& OutValue) const
+{
+	if (!InMaterialProperty.IsValid() || InMaterialProperty->GetPropertyType() != EDatasmithKeyValuePropertyType::String )
+	{
+		return false;
+	}
+
+	OutValue = InMaterialProperty->GetValue();
+
+	return true;
+}
diff --git a/Engine/Plugins/Enterprise/DatasmithImporter/Source/DatasmithTranslator/Private/MasterMaterials/DatasmithRevitMaterialSelector.cpp b/Engine/Plugins/Enterprise/DatasmithImporter/Source/DatasmithTranslator/Private/MasterMaterials/DatasmithRevitMaterialSelector.cpp
new file mode 100644
--- /dev/null
+++ b/Engine/Plugins/Enterprise/DatasmithImporter/Source/DatasmithTranslator/Private/MasterMaterials/DatasmithRevitMaterialSelector.cpp
@@ -0,0 +1,47 @@
+// Copyright 1998-2019 Epic Games, Inc. All Rights Reserved.
+
+#include "DatasmithRevitMaterialSelector.h"
+
+#include "IDatasmithSceneElements.h"
+
+#include "Materials/Material.h"
+#include "Materials/MaterialInstanceConstant.h"
+#include "UObject/SoftObjectPath.h"
+#include "Templates/Casts.h"
+
+FDatasmithRevitMaterialSelector::FDatasmithRevitMaterialSelector()
+{
+	// Master
+	MasterMaterial.FromSoftObjectPath( FSoftObjectPath("/DatasmithContent/Materials/RevitMaster.RevitMaster") );
+}
+
+bool FDatasmithRevitMaterialSelector::IsValid() const
+{
+	return MasterMaterial.IsValid();
+}
+
+const FDatasmithMasterMaterial& FDatasmithRevitMaterialSelector::GetMasterMaterial( const TSharedPtr< IDatasmithMasterMaterialElement >& /*InDatasmithMaterial*/ ) const
+{
+	return MasterMaterial;
+}
+
+void FDatasmithRevitMaterialSelector::FinalizeMaterialInstance(const TSharedPtr<IDatasmithMasterMaterialElement>& InDatasmithMaterial, UMaterialInstanceConstant * MaterialInstance) const
+{
+	if (!InDatasmithMaterial.IsValid() || MaterialInstance == nullptr)
+	{
+		return;
+	}
+
+	// Set blend mode to translucent if material requires transparency
+	if (InDatasmithMaterial->GetMaterialType() == EDatasmithMasterMaterialType::Transparent)
+	{
+		MaterialInstance->BasePropertyOverrides.bOverride_BlendMode = true;
+		MaterialInstance->BasePropertyOverrides.BlendMode = BLEND_Translucent;
+	}
+	// Set blend mode to masked if material has cutouts
+	else if (InDatasmithMaterial->GetMaterialType() == EDatasmithMasterMaterialType::CutOut)
+	{
+		MaterialInstance->BasePropertyOverrides.bOverride_BlendMode = true;
+		MaterialInstance->BasePropertyOverrides.BlendMode = BLEND_Masked;
+	}
+}
diff --git a/Engine/Plugins/Enterprise/DatasmithImporter/Source/DatasmithTranslator/Private/MasterMaterials/DatasmithRevitMaterialSelector.h b/Engine/Plugins/Enterprise/DatasmithImporter/Source/DatasmithTranslator/Private/MasterMaterials/DatasmithRevitMaterialSelector.h
new file mode 100644
--- /dev/null
+++ b/Engine/Plugins/Enterprise/DatasmithImporter/Source/DatasmithTranslator/Private/MasterMaterials/DatasmithRevitMaterialSelector.h
@@ -0,0 +1,24 @@
+// Copyright 1998-2019 Epic Games, Inc. All Rights Reserved.
+#pragma once
+
+#include "DatasmithDefinitions.h"
+#include "MasterMaterials/DatasmithMasterMaterial.h"
+#include "MasterMaterials/DatasmithMasterMaterialSelector.h"
+
+#include "Templates/SharedPointer.h"
+
+class IDatasmithMasterMaterialElement;
+
+class FDatasmithRevitMaterialSelector : public FDatasmithMasterMaterialSelector
+{
+public:
+	FDatasmithRevitMaterialSelector();
+	virtual ~FDatasmithRevitMaterialSelector() = default;
+
+	virtual bool IsValid() const override;
+	virtual const FDatasmithMasterMaterial& GetMasterMaterial( const TSharedPtr< IDatasmithMasterMaterialElement >& /*InDatasmithMaterial*/ ) const override;
+	virtual void FinalizeMaterialInstance(const TSharedPtr< IDatasmithMasterMaterialElement >& InDatasmithMaterial, UMaterialInstanceConstant* MaterialInstance) const override;
+
+private:
+	FDatasmithMasterMaterial MasterMaterial;
+};
diff --git a/Engine/Plugins/Enterprise/DatasmithImporter/Source/DatasmithTranslator/Private/MasterMaterials/DatasmithSketchupMaterialSelector.cpp b/Engine/Plugins/Enterprise/DatasmithImporter/Source/DatasmithTranslator/Private/MasterMaterials/DatasmithSketchupMaterialSelector.cpp
new file mode 100644
--- /dev/null
+++ b/Engine/Plugins/Enterprise/DatasmithImporter/Source/DatasmithTranslator/Private/MasterMaterials/DatasmithSketchupMaterialSelector.cpp
@@ -0,0 +1,40 @@
+// Copyright 1998-2019 Epic Games, Inc. All Rights Reserved.
+
+#include "DatasmithSketchupMaterialSelector.h"
+
+#include "IDatasmithSceneElements.h"
+
+#include "Materials/Material.h"
+#include "Materials/MaterialInstanceConstant.h"
+#include "UObject/SoftObjectPath.h"
+#include "Templates/Casts.h"
+
+FDatasmithSketchUpMaterialSelector::FDatasmithSketchUpMaterialSelector()
+{
+	// Master
+	MasterMaterial.FromSoftObjectPath( FSoftObjectPath("/DatasmithContent/Materials/SketchupMaster.SketchupMaster") );
+}
+
+bool FDatasmithSketchUpMaterialSelector::IsValid() const
+{
+	return MasterMaterial.IsValid();
+}
+
+const FDatasmithMasterMaterial& FDatasmithSketchUpMaterialSelector::GetMasterMaterial( const TSharedPtr< IDatasmithMasterMaterialElement >& /*InDatasmithMaterial*/ ) const
+{
+	return MasterMaterial;
+}
+
+void FDatasmithSketchUpMaterialSelector::FinalizeMaterialInstance(const TSharedPtr<IDatasmithMasterMaterialElement>& InDatasmithMaterial, UMaterialInstanceConstant * MaterialInstance) const
+{
+	if (!InDatasmithMaterial.IsValid() || MaterialInstance == nullptr)
+	{
+		return;
+	}
+
+	if (InDatasmithMaterial->GetMaterialType() == EDatasmithMasterMaterialType::Transparent)
+	{
+		MaterialInstance->BasePropertyOverrides.bOverride_BlendMode = true;
+		MaterialInstance->BasePropertyOverrides.BlendMode = BLEND_Translucent;
+	}
+}
diff --git a/Engine/Plugins/Enterprise/DatasmithImporter/Source/DatasmithTranslator/Private/MasterMaterials/DatasmithSketchupMaterialSelector.h b/Engine/Plugins/Enterprise/DatasmithImporter/Source/DatasmithTranslator/Private/MasterMaterials/DatasmithSketchupMaterialSelector.h
new file mode 100644
--- /dev/null
+++ b/Engine/Plugins/Enterprise/DatasmithImporter/Source/DatasmithTranslator/Private/MasterMaterials/DatasmithSketchupMaterialSelector.h
@@ -0,0 +1,24 @@
+// Copyright 1998-2019 Epic Games, Inc. All Rights Reserved.
+#pragma once
+
+#include "DatasmithDefinitions.h"
+#include "MasterMaterials/DatasmithMasterMaterial.h"
+#include "MasterMaterials/DatasmithMasterMaterialSelector.h"
+
+#include "Templates/SharedPointer.h"
+
+class IDatasmithMasterMaterialElement;
+
+class FDatasmithSketchUpMaterialSelector : public FDatasmithMasterMaterialSelector
+{
+public:
+	FDatasmithSketchUpMaterialSelector();
+	virtual ~FDatasmithSketchUpMaterialSelector() = default;
+
+	virtual bool IsValid() const override;
+	virtual const FDatasmithMasterMaterial& GetMasterMaterial( const TSharedPtr< IDatasmithMasterMaterialElement >& /*InDatasmithMaterial*/ ) const override;
+	virtual void FinalizeMaterialInstance(const TSharedPtr< IDatasmithMasterMaterialElement >& InDatasmithMaterial, UMaterialInstanceConstant* MaterialInstance) const override;
+
+private:
+	FDatasmithMasterMaterial MasterMaterial;
+};
diff --git a/Engine/Plugins/Enterprise/DatasmithImporter/Source/DatasmithTranslator/Private/Utility/DatasmithMathUtils.cpp b/Engine/Plugins/Enterprise/DatasmithImporter/Source/DatasmithTranslator/Private/Utility/DatasmithMathUtils.cpp
new file mode 100644
--- /dev/null
+++ b/Engine/Plugins/Enterprise/DatasmithImporter/Source/DatasmithTranslator/Private/Utility/DatasmithMathUtils.cpp
@@ -0,0 +1,35 @@
+// Copyright 1998-2019 Epic Games, Inc. All Rights Reserved.
+
+#include "Utility/DatasmithMathUtils.h"
+
+#include "Math/UnrealMathUtility.h"
+#include "Math/TransformVectorized.h"
+#include "Math/Quat.h"
+
+
+void FDatasmithTransformUtils::GetSafeRotation(const FTransform& InTransform, FQuat& OutRotation)
+{
+	OutRotation = InTransform.GetRotation();
+
+	// Workaround to floating point precision issues with quaternion
+	{
+		// this value was found from ONE experience, so I'm waiting on other problem to improve it
+		if (FMath::Abs(OutRotation.X) < 0.005f)
+		{
+			OutRotation.X = 0.f;
+		}
+		if (FMath::Abs(OutRotation.Y) < 0.005f)
+		{
+			OutRotation.Y = 0.f;
+		}
+		if (FMath::Abs(OutRotation.Z) < 0.005f)
+		{
+			OutRotation.Z = 0.f;
+		}
+		if ((1 - FMath::Abs(OutRotation.W)) < 0.001f)
+		{
+			OutRotation.W = (OutRotation.W > 0) ? 1.f : -1.f;
+		}
+		OutRotation.Normalize();
+	}
+}
\ No newline at end of file
diff --git a/Engine/Plugins/Enterprise/DatasmithImporter/Source/DatasmithTranslator/Private/Utility/DatasmithMeshHelper.cpp b/Engine/Plugins/Enterprise/DatasmithImporter/Source/DatasmithTranslator/Private/Utility/DatasmithMeshHelper.cpp
new file mode 100644
--- /dev/null
+++ b/Engine/Plugins/Enterprise/DatasmithImporter/Source/DatasmithTranslator/Private/Utility/DatasmithMeshHelper.cpp
@@ -0,0 +1,381 @@
+// Copyright 1998-2019 Epic Games, Inc. All Rights Reserved.
+
+#include "Utility/DatasmithMeshHelper.h"
+
+#include "StaticMeshAttributes.h"
+#include "Engine/StaticMesh.h"
+#include "MeshDescriptionOperations.h"
+#include "DatasmithMesh.h"
+#include "RawMesh.h"
+#include "DatasmithUtils.h"
+#include "DatasmithMeshUObject.h"
+#include "Misc/SecureHash.h"
+#include "UVMapSettings.h"
+
+
+namespace DatasmithMeshHelper
+{
+
+	int32 GetPolygonCount(const FMeshDescription& Mesh)
+	{
+		return Mesh.Polygons().Num();
+	}
+
+	int32 GetTriangleCount(const FMeshDescription& Mesh)
+	{
+		return Mesh.Triangles().Num();
+	}
+
+	void ExtractVertexPositions(const FMeshDescription& Mesh, TArray<FVector>& Positions)
+	{
+		FStaticMeshConstAttributes Attributes(Mesh);
+		const TVertexAttributesConstRef<FVector> VertexPositions = Attributes.GetVertexPositions();
+		if (VertexPositions.IsValid())
+		{
+			Positions.Reserve(VertexPositions.GetNumElements());
+			for (const FVertexID VertexID : Mesh.Vertices().GetElementIDs())
+			{
+				Positions.Add(VertexPositions[VertexID]);
+			}
+		}
+	}
+
+	void PrepareAttributeForStaticMesh(FMeshDescription& Mesh)
+	{
+		FStaticMeshAttributes(Mesh).Register();
+	}
+
+	bool IsTriangleDegenerated(const FMeshDescription& Mesh, const FMeshTriangle& MeshTriangle)
+	{
+		// extract triangle vertices positions
+		FVector P[3];
+		const TVertexAttributesConstRef<FVector> VertexPositions = Mesh.VertexAttributes().GetAttributesRef<FVector>(MeshAttribute::Vertex::Position);
+		for (int32 CornerIndex = 0; CornerIndex < 3; ++CornerIndex)
+		{
+			FVertexID VertexID = Mesh.GetVertexInstanceVertex(MeshTriangle.GetVertexInstanceID(CornerIndex));
+			P[CornerIndex] = VertexPositions[VertexID];
+		}
+
+		float NormalLengthSquare = ((P[1] - P[2]) ^ (P[0] - P[2])).SizeSquared();
+		return NormalLengthSquare < SMALL_NUMBER;
+	}
+
+	void RemoveEmptyPolygonGroups(FMeshDescription& Mesh)
+	{
+		bool bRemovedSection = false;
+		for (const FPolygonGroupID& PolygonGroupID : Mesh.PolygonGroups().GetElementIDs())
+		{
+			if (Mesh.GetNumPolygonGroupPolygons(PolygonGroupID) == 0)
+			{
+				Mesh.DeletePolygonGroup(PolygonGroupID);
+				bRemovedSection = true;
+			}
+		}
+
+		if (bRemovedSection)
+		{
+			FElementIDRemappings Remappings;
+			Mesh.Compact(Remappings);
+		}
+	}
+
+#if WITH_EDITOR
+	FMeshDescription* InitMeshDescription(UStaticMesh* StaticMesh, int32 LodIndex)
+	{
+		if (!ensure(StaticMesh))
+		{
+			return nullptr;
+		}
+
+		while (!StaticMesh->IsSourceModelValid(LodIndex))
+		{
+			StaticMesh->AddSourceModel();
+		}
+
+		return StaticMesh->CreateMeshDescription(LodIndex);
+	}
+
+	void FillUStaticMesh(UStaticMesh* StaticMesh, int32 LodIndex, FRawMesh& RawMesh, const TMap<int32, FName>* InMaterialMapInverse)
+	{
+		FMeshDescription MeshDescription;
+		PrepareAttributeForStaticMesh(MeshDescription);
+		if (InMaterialMapInverse == nullptr)
+		{
+			// no explicit mapping, build one from UStaticMesh
+			TMap<FName, int32> MaterialMap; // unused
+			TMap<int32, FName> MaterialMapInverse;
+
+ 			BuildMaterialMappingFromStaticMesh(StaticMesh, MaterialMap, MaterialMapInverse);
+			FMeshDescriptionOperations::ConvertFromRawMesh(RawMesh, MeshDescription, MaterialMapInverse);
+		}
+		else
+		{
+			FMeshDescriptionOperations::ConvertFromRawMesh(RawMesh, MeshDescription, *InMaterialMapInverse);
+		}
+
+		FillUStaticMesh(StaticMesh, LodIndex, MoveTemp(MeshDescription));
+	}
+
+	void FillUStaticMesh(UStaticMesh* StaticMesh, int32 LodIndex, FDatasmithMesh& DsMesh)
+	{
+		FMeshDescription MeshDescription;
+		PrepareAttributeForStaticMesh(MeshDescription);
+		if (FDatasmithMeshUtils::ToMeshDescription(DsMesh, MeshDescription))
+		{
+			FillUStaticMesh(StaticMesh, LodIndex, MoveTemp(MeshDescription));
+		}
+	}
+
+	void FillUStaticMesh(UStaticMesh* StaticMesh, int LodIndex, FMeshDescription&& MeshDescription)
+	{
+		if ( !IsMeshValid( MeshDescription ) )
+		{
+			return;
+		}
+
+		FMeshDescription* OriginalMeshDescription = InitMeshDescription(StaticMesh, LodIndex);
+		if (ensure(OriginalMeshDescription))
+		{
+			// #ueent_todo: validate attributes too make sure we don't push incompatible MeshDescriptions
+			*OriginalMeshDescription = MoveTemp(MeshDescription);
+		}
+	}
+
+	void FillUStaticMeshByCopy(UStaticMesh* StaticMesh, int LodIndex, FMeshDescription MeshDescription)
+	{
+		FillUStaticMesh(StaticMesh, LodIndex, MoveTemp(MeshDescription));
+	}
+
+	void PrepareStaticMaterials(UStaticMesh* StaticMesh, int32 MaterialCount)
+	{
+		if (!ensure(StaticMesh))
+		{
+			return;
+		}
+
+		MaterialCount = FMath::Max(MaterialCount, 1);
+
+		// Don't touch the StaticMesh if it's not needed (would break the template system on reimport)
+		if (MaterialCount != StaticMesh->StaticMaterials.Num())
+		{
+			StaticMesh->StaticMaterials.Empty(MaterialCount);
+			// Declare materials applied to the StaticMesh according to fit the logic in UDatasmithImportFactory::UpdateMaterials
+			for (int MaterialIndex = 0; MaterialIndex < MaterialCount; MaterialIndex++)
+			{
+				StaticMesh->GetSectionInfoMap().Set(0, MaterialIndex, FMeshSectionInfo(MaterialIndex));
+				FName SlotName = DefaultSlotName(MaterialIndex);
+				StaticMesh->StaticMaterials.Add(FStaticMaterial(nullptr, SlotName, SlotName));
+			}
+		}
+	}
+
+	void BuildMaterialMappingFromStaticMesh(UStaticMesh* StaticMesh, TMap<FName, int32>& MaterialMap, TMap<int32, FName>& MaterialMapInverse)
+	{
+		if (!ensure(StaticMesh))
+		{
+			return;
+		}
+
+		int32 MatCount = StaticMesh->StaticMaterials.Num();
+		ensure(MatCount);
+		MaterialMap.Empty(MatCount);
+		MaterialMapInverse.Empty(MatCount);
+
+		for ( int32 MaterialIndex = 0; MaterialIndex < MatCount; ++MaterialIndex )
+		{
+			FName& MatName = StaticMesh->StaticMaterials[MaterialIndex].ImportedMaterialSlotName;
+			MaterialMap.Add(MatName, MaterialIndex);
+			MaterialMapInverse.Add(MaterialIndex, MatName);
+		}
+	}
+
+	void HashMeshLOD(UStaticMesh* StaticMesh, int32 LodIndex, FMD5& MD5)
+	{
+		if (ensure(StaticMesh))
+		{
+			const FMeshDescription* Mesh = StaticMesh->GetMeshDescription(LodIndex);
+			if (ensure(Mesh))
+			{
+				HashMeshDescription(*Mesh, MD5);
+			}
+		}
+	}
+#endif //WITH_EDITOR
+
+	void HashMeshDescription(const FMeshDescription& Mesh, FMD5& MD5)
+	{
+		auto HashAttributeSet = [](FMD5& MD5, const FAttributesSetBase& AttributeSet)
+		{
+			TArray<FName> OutAttributeNames;
+			AttributeSet.GetAttributeNames(OutAttributeNames);
+			for (const FName& AttributeName: OutAttributeNames)
+			{
+				int32 AttributeHash = AttributeSet.GetHash(AttributeName);
+				MD5.Update((uint8*)&AttributeHash, sizeof(AttributeHash));
+				// #ueent_todo better hash function...
+				// see FDatasmithFBXSceneMesh::GetHash
+			}
+		};
+
+		HashAttributeSet(MD5, Mesh.VertexAttributes());
+		HashAttributeSet(MD5, Mesh.VertexInstanceAttributes());
+		HashAttributeSet(MD5, Mesh.EdgeAttributes());
+		HashAttributeSet(MD5, Mesh.PolygonAttributes());
+		HashAttributeSet(MD5, Mesh.PolygonGroupAttributes());
+	}
+
+	FName DefaultSlotName(int32 MaterialIndex)
+	{
+		return FName( *FString::FromInt(MaterialIndex) );
+	}
+
+	int32 GetNumUVChannel(const FMeshDescription& Mesh)
+	{
+		const TVertexInstanceAttributesConstRef<FVector2D> UVChannels = Mesh.VertexInstanceAttributes().GetAttributesRef<FVector2D>(MeshAttribute::VertexInstance::TextureCoordinate);
+		return UVChannels.GetNumIndices();
+	}
+
+	bool HasUVChannel(const FMeshDescription& Mesh, int32 ChannelIndex)
+	{
+		return ChannelIndex >= 0 && ChannelIndex < GetNumUVChannel(Mesh);
+	}
+
+	bool HasUVData(const FMeshDescription& Mesh, int32 ChannelIndex)
+	{
+		if (HasUVChannel(Mesh, ChannelIndex))
+		{
+			const TVertexInstanceAttributesConstRef<FVector2D> UVChannels = Mesh.VertexInstanceAttributes().GetAttributesRef<FVector2D>(MeshAttribute::VertexInstance::TextureCoordinate);
+			const FVector2D DefValue = UVChannels.GetDefaultValue();
+			for (FVertexInstanceID InstanceID : Mesh.VertexInstances().GetElementIDs())
+			{
+				if (UVChannels.Get(InstanceID, ChannelIndex) != DefValue)
+				{
+					return true;
+				}
+			}
+		}
+		return false;
+	}
+
+	bool RequireUVChannel(FMeshDescription& Mesh, int32 ChannelIndex)
+	{
+		if (ChannelIndex < 0 || ChannelIndex >= MAX_MESH_TEXTURE_COORDS_MD)
+		{
+			return false;
+		}
+
+		if (!HasUVChannel(Mesh, ChannelIndex))
+		{
+			TVertexInstanceAttributesRef<FVector2D> UVChannels = Mesh.VertexInstanceAttributes().GetAttributesRef<FVector2D>(MeshAttribute::VertexInstance::TextureCoordinate);
+			UVChannels.SetNumIndices(ChannelIndex + 1);
+		}
+		return true;
+	}
+
+	void CreateDefaultUVs(FDatasmithMesh& DatasmithMesh)
+	{
+		// Get the mesh description to generate BoxUV.
+		FMeshDescription MeshDescription;
+		PrepareAttributeForStaticMesh(MeshDescription);
+		FDatasmithMeshUtils::ToMeshDescription(DatasmithMesh, MeshDescription);
+		FUVMapParameters UVParameters(DatasmithMesh.GetExtents().GetCenter(), FQuat::Identity, DatasmithMesh.GetExtents().GetSize(), FVector::OneVector, FVector2D::UnitVector);
+		TMap<FVertexInstanceID, FVector2D> TexCoords;
+		FMeshDescriptionOperations::GenerateBoxUV(MeshDescription, UVParameters, TexCoords);
+
+		// Put the results in a map to determine the number of unique values.
+		TMap<FVector2D, TArray<int32>> UniqueTexCoordMap;
+		for (const TPair<FVertexInstanceID, FVector2D>& Pair : TexCoords)
+		{
+			TArray<int32>& MappedIndices = UniqueTexCoordMap.FindOrAdd(Pair.Value);
+			MappedIndices.Add(Pair.Key.GetValue());
+		}
+
+		//Set the UV values
+		if (DatasmithMesh.GetUVChannelsCount() == 0)
+		{
+			DatasmithMesh.AddUVChannel();
+		}
+		DatasmithMesh.SetUVCount(0, UniqueTexCoordMap.Num());
+		int32 UVIndex = 0;
+		TArray<int32> IndicesMapping;
+		IndicesMapping.AddZeroed(TexCoords.Num());
+
+		for (const TPair<FVector2D, TArray<int32>>& UniqueCoordPair : UniqueTexCoordMap)
+		{
+			DatasmithMesh.SetUV(0, UVIndex, UniqueCoordPair.Key.X, UniqueCoordPair.Key.Y);
+			for (int32 IndicesIndex : UniqueCoordPair.Value)
+			{
+				IndicesMapping[IndicesIndex] = UVIndex;
+			}
+			UVIndex++;
+		}
+
+		//Map the UV indices.
+		for (int32 FaceIndex = 0; FaceIndex < DatasmithMesh.GetFacesCount(); ++FaceIndex)
+		{
+			const int32 IndicesOffset = FaceIndex * 3;
+			check(IndicesOffset + 2 < IndicesMapping.Num());
+			DatasmithMesh.SetFaceUV(FaceIndex, 0, IndicesMapping[IndicesOffset + 0], IndicesMapping[IndicesOffset + 1], IndicesMapping[IndicesOffset + 2]);
+		}
+	}
+
+	void CreateDefaultUVs(FMeshDescription& MeshDescription)
+	{
+		FBox MeshBoundingBox = MeshDescription.ComputeBoundingBox();
+		FUVMapParameters UVParameters(MeshBoundingBox.GetCenter(), FQuat::Identity, MeshBoundingBox.GetExtent(), FVector::OneVector, FVector2D::UnitVector);
+		TMap<FVertexInstanceID, FVector2D> TexCoords;
+		FMeshDescriptionOperations::GenerateBoxUV(MeshDescription, UVParameters, TexCoords);
+
+		TMeshAttributesRef<FVertexInstanceID, FVector2D> UVs = MeshDescription.VertexInstanceAttributes().GetAttributesRef<FVector2D>(MeshAttribute::VertexInstance::TextureCoordinate);
+		if (UVs.GetNumIndices() == 0)
+		{
+			UVs.SetNumIndices(1);
+		}
+
+		for (const FVertexInstanceID& VertexInstanceID : MeshDescription.VertexInstances().GetElementIDs())
+		{
+			if (const FVector2D* UVCoord = TexCoords.Find(VertexInstanceID))
+			{
+				UVs.Set(VertexInstanceID, 0, *UVCoord);
+			}
+			else
+			{
+				ensureMsgf(false, TEXT("Tried to apply UV data that did not match the MeshDescription."));
+			}
+		}
+	}
+
+	bool IsMeshValid(const FMeshDescription& Mesh, FVector BuildScale)
+	{
+		const TVertexAttributesConstRef<FVector> VertexPositions = Mesh.VertexAttributes().GetAttributesRef<FVector>(MeshAttribute::Vertex::Position);
+		FVector RawNormalScale(BuildScale.Y*BuildScale.Z, BuildScale.X*BuildScale.Z, BuildScale.X*BuildScale.Y); // Component-wise scale
+		for (const FPolygonID PolygonID : Mesh.Polygons().GetElementIDs())
+		{
+			FVector Corners[3];
+			for (const FTriangleID TriangleID : Mesh.GetPolygonTriangleIDs(PolygonID))
+			{
+				for (int32 CornerIndex = 0; CornerIndex < 3; ++CornerIndex)
+				{
+					FVertexInstanceID VertexInstanceID = Mesh.GetTriangleVertexInstance(TriangleID, CornerIndex);
+					FVertexID VertexID = Mesh.GetVertexInstanceVertex(VertexInstanceID);
+					Corners[CornerIndex] = VertexPositions[VertexID];
+				}
+
+				FVector RawNormal = (Corners[1] - Corners[2]) ^ (Corners[0] - Corners[2]);
+				RawNormal *= RawNormalScale;
+				double FourSquaredTriangleArea = RawNormal.SizeSquared();
+
+				// We support even small triangles, but this function is still useful to
+				// see if we have at least one valid triangle in the mesh
+				if (FourSquaredTriangleArea > 0.0f)
+				{
+					return true;
+				}
+			}
+		}
+
+		// all faces are degenerated, mesh is invalid
+		return false;
+	}
+
+} // namespace DatasmithMeshHelper
diff --git a/Engine/Plugins/Enterprise/DatasmithImporter/Source/DatasmithTranslator/Public/DatasmithTranslatableSource.h b/Engine/Plugins/Enterprise/DatasmithImporter/Source/DatasmithTranslator/Public/DatasmithTranslatableSource.h
new file mode 100644
--- /dev/null
+++ b/Engine/Plugins/Enterprise/DatasmithImporter/Source/DatasmithTranslator/Public/DatasmithTranslatableSource.h
@@ -0,0 +1,34 @@
+// Copyright 1998-2019 Epic Games, Inc. All Rights Reserved.
+
+#pragma once
+
+#include "DatasmithSceneSource.h"
+
+#include "CoreMinimal.h"
+#include "UObject/StrongObjectPtr.h"
+
+/**
+ * Wrap a source with an adapted translator.
+ * This scopes the lifecycle of a translator, 
+ */
+struct DATASMITHTRANSLATOR_API FDatasmithTranslatableSceneSource
+{
+	FDatasmithTranslatableSceneSource(const FDatasmithSceneSource& Source);
+	~FDatasmithTranslatableSceneSource();
+
+	bool IsTranslatable() const;
+
+	bool Translate(TSharedRef< class IDatasmithScene > Scene);
+
+	TSharedPtr<class IDatasmithTranslator> GetTranslator() const;
+
+private:
+	/** Translator currently in use (null when not importing) */
+	TSharedPtr<class IDatasmithTranslator> Translator;
+
+	/** internal helper to release scene */
+	TUniquePtr< class FSceneGuard > SceneGuard;
+};
+
+
+
diff --git a/Engine/Plugins/Enterprise/DatasmithImporter/Source/DatasmithTranslator/Public/MasterMaterials/DatasmithMasterMaterial.h b/Engine/Plugins/Enterprise/DatasmithImporter/Source/DatasmithTranslator/Public/MasterMaterials/DatasmithMasterMaterial.h
new file mode 100644
--- /dev/null
+++ b/Engine/Plugins/Enterprise/DatasmithImporter/Source/DatasmithTranslator/Public/MasterMaterials/DatasmithMasterMaterial.h
@@ -0,0 +1,33 @@
+// Copyright 1998-2019 Epic Games, Inc. All Rights Reserved.
+#pragma once
+
+#include "Containers/Set.h"
+#include "Containers/UnrealString.h"
+
+struct FSoftObjectPath;
+class UMaterial;
+
+class DATASMITHTRANSLATOR_API FDatasmithMasterMaterial
+{
+public:
+	FDatasmithMasterMaterial();
+	explicit FDatasmithMasterMaterial( UMaterial* InMaterial );
+
+	void FromMaterial( UMaterial* InMaterial );
+	void FromSoftObjectPath( const FSoftObjectPath& InObjectPath );
+
+	UMaterial* GetMaterial() const { return Material; }
+
+	bool IsValid() const { return Material != nullptr; }
+
+#if WITH_EDITOR
+	TSet< FString > VectorParams;
+	TSet< FString > ScalarParams;
+	TSet< FString > TextureParams;
+	TSet< FString > BoolParams;
+#endif	
+
+private:
+	UMaterial* Material;
+
+};
diff --git a/Engine/Plugins/Enterprise/DatasmithImporter/Source/DatasmithTranslator/Public/MasterMaterials/DatasmithMasterMaterialManager.h b/Engine/Plugins/Enterprise/DatasmithImporter/Source/DatasmithTranslator/Public/MasterMaterials/DatasmithMasterMaterialManager.h
new file mode 100644
--- /dev/null
+++ b/Engine/Plugins/Enterprise/DatasmithImporter/Source/DatasmithTranslator/Public/MasterMaterials/DatasmithMasterMaterialManager.h
@@ -0,0 +1,28 @@
+// Copyright 1998-2019 Epic Games, Inc. All Rights Reserved.
+#pragma once
+
+#include "Containers/ContainersFwd.h"
+#include "Templates/SharedPointer.h"
+#include "Templates/UniquePtr.h"
+
+class FDatasmithMasterMaterialSelector;
+
+class DATASMITHTRANSLATOR_API FDatasmithMasterMaterialManager
+{
+public:
+	static void Create();
+	static void Destroy();
+	static FDatasmithMasterMaterialManager& Get();
+
+	FString GetHostFromString( const FString& HostString );
+
+	void RegisterSelector( const TCHAR* Host, TSharedPtr< FDatasmithMasterMaterialSelector > Selector );
+	void UnregisterSelector( const TCHAR* Host );
+
+	const TSharedPtr< FDatasmithMasterMaterialSelector > GetSelector( const TCHAR* Host ) const;
+
+private:
+	static TUniquePtr< FDatasmithMasterMaterialManager > Instance;
+
+	TMap< FString, TSharedPtr< FDatasmithMasterMaterialSelector > > Selectors;
+};
\ No newline at end of file
diff --git a/Engine/Plugins/Enterprise/DatasmithImporter/Source/DatasmithTranslator/Public/MasterMaterials/DatasmithMasterMaterialSelector.h b/Engine/Plugins/Enterprise/DatasmithImporter/Source/DatasmithTranslator/Public/MasterMaterials/DatasmithMasterMaterialSelector.h
new file mode 100644
--- /dev/null
+++ b/Engine/Plugins/Enterprise/DatasmithImporter/Source/DatasmithTranslator/Public/MasterMaterials/DatasmithMasterMaterialSelector.h
@@ -0,0 +1,31 @@
+// Copyright 1998-2019 Epic Games, Inc. All Rights Reserved.
+#pragma once
+
+#include "Math/Color.h"
+#include "Templates/SharedPointer.h"
+
+class FDatasmithMasterMaterial;
+class IDatasmithMasterMaterialElement;
+class IDatasmithKeyValueProperty;
+class UMaterialInstanceConstant;
+
+class DATASMITHTRANSLATOR_API FDatasmithMasterMaterialSelector
+{
+public:
+	virtual ~FDatasmithMasterMaterialSelector() = default;
+
+	virtual bool IsValid() const { return false; }
+	virtual const FDatasmithMasterMaterial& GetMasterMaterial( const TSharedPtr< IDatasmithMasterMaterialElement >& InDatasmithMaterial ) const;
+	virtual void FinalizeMaterialInstance(const TSharedPtr< IDatasmithMasterMaterialElement >& InDatasmithMaterial, UMaterialInstanceConstant* MaterialInstance) const {}
+
+	virtual bool GetColor( const TSharedPtr< IDatasmithKeyValueProperty >& InMaterialProperty, FLinearColor& OutValue ) const;
+	virtual bool GetFloat( const TSharedPtr< IDatasmithKeyValueProperty >& InMaterialProperty, float& OutValue ) const;
+	virtual bool GetBool( const TSharedPtr< IDatasmithKeyValueProperty >& InMaterialProperty, bool& OutValue ) const;
+	virtual bool GetTexture( const TSharedPtr< IDatasmithKeyValueProperty >& InMaterialProperty, FString& OutValue ) const;
+	virtual bool GetString( const TSharedPtr< IDatasmithKeyValueProperty >& InMaterialProperty, FString& OutValue ) const;
+
+protected:
+	static FDatasmithMasterMaterial InvalidMasterMaterial;
+
+};
+
diff --git a/Engine/Plugins/Enterprise/DatasmithImporter/Source/DatasmithTranslator/Public/Utility/DatasmithMathUtils.h b/Engine/Plugins/Enterprise/DatasmithImporter/Source/DatasmithTranslator/Public/Utility/DatasmithMathUtils.h
new file mode 100644
--- /dev/null
+++ b/Engine/Plugins/Enterprise/DatasmithImporter/Source/DatasmithTranslator/Public/Utility/DatasmithMathUtils.h
@@ -0,0 +1,18 @@
+// Copyright 1998-2019 Epic Games, Inc. All Rights Reserved.
+#pragma once
+
+#include "CoreMinimal.h"
+
+
+struct FTransform;
+struct FQuat;
+
+class DATASMITHTRANSLATOR_API FDatasmithTransformUtils
+{
+public:
+
+	/**
+	 * Get rotation from FTransform and fix floating point precision issues with quaternion
+	 */
+	static void GetSafeRotation(const FTransform& InTransform, FQuat& OutRotation);
+};
\ No newline at end of file
diff --git a/Engine/Plugins/Experimental/Text3D/Source/Text3D/Private/Contour.cpp b/Engine/Plugins/Experimental/Text3D/Source/Text3D/Private/Contour.cpp
new file mode 100644
--- /dev/null
+++ b/Engine/Plugins/Experimental/Text3D/Source/Text3D/Private/Contour.cpp
@@ -0,0 +1,162 @@
+// Copyright 1998-2019 Epic Games, Inc. All Rights Reserved.
+
+
+#include "Contour.h"
+#include "Part.h"
+#include "Data.h"
+
+
+FContour::FContour(const TSharedPtr<FData> DataIn)
+	: Data(DataIn)
+{
+	ResetContour();
+}
+
+FContour::~FContour()
+{
+	for (const FPart* const Part : *this)
+	{
+		delete Part;
+	}
+}
+
+int32 FContour::GetPrev(const int32 Index) const
+{
+	return (Index + Num() - 1) % Num();
+}
+
+int32 FContour::GetNext(const int32 Index) const
+{
+	return (Index + 1) % Num();
+}
+
+bool FContour::HasIntersections() const
+{
+	return bHasIntersections;
+}
+
+void FContour::ResetContour()
+{
+	bHasIntersections = true;
+}
+
+void FContour::DisableIntersections()
+{
+	bHasIntersections = false;
+}
+
+// p_1 ~ Point->InitialPosition
+// p_2 ~ Point->Next->InitialPosition
+// e   ~ Expand
+// d   ~ Point->DoneExpand
+// t   ~ Point->TangentX
+// a = e - d
+// (b_i = p_i + n_i * a) is a position to which point will be expanded
+// (b_2 - b_1) has same direction as (t) if no intersection happened and opposite one if it did
+// if intersection happened, needed value can be received from (b_2 - b_1 = 0)
+void FContour::ComputeAvailableExpandNear(FPart* const Point)
+{
+	FPart* Next = Point->Next;
+
+	const FVector2D dp = Next->InitialPosition - Point->InitialPosition;
+	const FVector2D dn = Next->Normal - Point->Normal;
+
+	// ExpandTotal is used instead of Expand to compute AvailableExpandNear once for all segments, without recomputing them at every BevelLinear call.
+	const float Expand = Data->GetExpandTotal();
+
+	// (2 * Expand) is used to mark this intersection invalid ((Expand) is total needed expand)
+	Point->AvailableExpandNear = FVector2D::DotProduct(Point->TangentX, dp + dn * Expand) < 0 ? dp.Size() / dn.Size() - Point->DoneExpand : 2 * Expand;
+}
+
+void FContour::ComputeAvailableExpandsFarFrom(FPart* const Point)
+{
+	Point->AvailableExpandsFar.Empty();
+
+	// IntersectionFar requires at leat 1 part distance (between point and edge) if edge if further counterclockwise then point and at least 2 parts distance if further clockwise (because FPart is point and it's _next_ edge)
+	for (const FPart* Edge = Point->Next->Next; Edge != Point->Prev->Prev; Edge = Edge->Next)
+	{
+		ComputeAvailableExpandFar(Point, Edge);
+	}
+}
+
+void FContour::ComputeAvailableExpandsFarTo(const FPart* const Edge)
+{
+	// See comment in ComputeAvailableExpandsFarFrom
+	for (FPart* Point = Edge->Next->Next->Next; Point != Edge->Prev; Point = Point->Next)
+	{
+		Point->AvailableExpandsFar.Remove(Edge);
+		ComputeAvailableExpandFar(Point, Edge);
+	}
+}
+
+void FContour::RemoveRange(const FPart* const Start, const FPart* const End)
+{
+	const int32 StartIndex = Find(Start);
+	const int32 EndIndex = Find(End);
+
+	if (EndIndex < StartIndex)
+	{
+		RemoveAt(StartIndex, Num() - StartIndex);
+		RemoveAt(0, EndIndex);
+	}
+	else
+	{
+		RemoveAt(StartIndex, EndIndex - StartIndex);
+	}
+}
+
+// p_1 ~ Edge->Position
+// p_2 ~ Edge->Next->Position
+// p_3 ~ Point->Position
+// n_2 ~ Edge->Next->Normal
+// n_3 ~ Point->Normal
+// d_2 ~ Edge->Next->DoneExpand
+// d_3 ~ Point->DoneExpand
+// e   ~ total expand
+// t   ~ EdgeA->TangentX
+// e_2 = e - d_2
+// e_3 = e - d_3
+// check if point's normal approaches edge from needed side with sign of cross product ([t, n_3])
+// b = d_3 - d_2
+// if intersection happens, for point of intersection ((p_2 + e_2 * n_2) - (p_3 + e_3 * n_3)) is parallel to (t)
+// so their cross-product is zero: ([(p_2 + e_2 * n_2) - (p_3 + e_3 * n_3), t] = 0)
+// then get needed value
+// if value is (<= 0), intersection will not happen
+// available expand for edge can be received from (e_2 + d_2 = e_3 + d_3)
+// check if intersection is _on_ expanded edge with cross products (previous operations guarantee only that it's on the line this edge belongs too, not enough to claim an intersection)
+void FContour::ComputeAvailableExpandFar(FPart* const Point, const FPart* const Edge)
+{
+	const FPart* const EdgeA = Edge;
+	const FPart* const EdgeB = Edge->Next;
+
+	const FVector2D dp = EdgeA->TangentX;
+	const float dpXPointNormal = FVector2D::CrossProduct(dp, Point->Normal);
+
+	if (dpXPointNormal <= 0)
+	{
+		return;
+	}
+
+	const float DoneExpandDiff = Point->DoneExpand - EdgeB->DoneExpand;
+	const float AvailableExpandPoint = FVector2D::CrossProduct(dp, (DoneExpandDiff * EdgeB->Normal - Point->Position + EdgeB->Position)) / (dpXPointNormal - FVector2D::CrossProduct(dp, EdgeB->Normal));
+
+	if (AvailableExpandPoint <= 0)
+	{
+		return;
+	}
+
+	const float AvailableExpandEdgeB = DoneExpandDiff + AvailableExpandPoint;
+	const FVector2D PointExpanded = Point->Expanded(AvailableExpandPoint) - EdgeB->Position;
+
+	if (FVector2D::CrossProduct(PointExpanded, (EdgeB->Expanded(AvailableExpandEdgeB) - EdgeB->Position)) < 0)
+	{
+		return;
+	}
+
+	if (FVector2D::CrossProduct(PointExpanded, (EdgeA->Expanded(EdgeB->DoneExpand + AvailableExpandEdgeB - EdgeA->DoneExpand) - EdgeB->Position)) > 0)
+	{
+		return;
+	}
+
+	Point->AvailableExpandsFar.Add(EdgeA, AvailableExpandPoint);
+}
diff --git a/Engine/Plugins/Experimental/Text3D/Source/Text3D/Private/Contour.h b/Engine/Plugins/Experimental/Text3D/Source/Text3D/Private/Contour.h
new file mode 100644
--- /dev/null
+++ b/Engine/Plugins/Experimental/Text3D/Source/Text3D/Private/Contour.h
@@ -0,0 +1,66 @@
+// Copyright 1998-2019 Epic Games, Inc. All Rights Reserved.
+
+#pragma once
+
+#include "CoreMinimal.h"
+#include "Util.h"
+
+struct FPart;
+class FData;
+
+class FContour final : public FContourBase
+{
+public:
+	struct FPathEntry
+	{
+		int32 Prev;
+		int32 Next;
+	};
+
+
+	FContour(const TSharedPtr<FData> DataIn);
+	~FContour();
+
+	int32 GetPrev(const int32 Index) const;
+	int32 GetNext(const int32 Index) const;
+
+	bool HasIntersections() const;
+	void ResetContour();
+	void DisableIntersections();
+
+	/**
+	 * Compute expand value at which point's normal will intersect next point's normal. Result is written to FPart::AvailableExpandNear.
+	 * @param Point - Point for which available expand should be computed.
+	 */
+	void ComputeAvailableExpandNear(FPart* const Point);
+
+	/**
+	 * Call FContour::ComputeAvailableExpandFar for given point and all edges.
+	 * @param Point - Point for which available expand should be computed.
+	 */
+	void ComputeAvailableExpandsFarFrom(FPart* const Point);
+	/**
+	 * Call FContour::ComputeAvailableExpandFar for all points and given edge.
+	 * @param Edge - Edge for which available expand should be computed.
+	 */
+	void ComputeAvailableExpandsFarTo(const FPart* const Edge);
+
+	/**
+	 * Remove multiple parts of contour.
+	 * @param Start - Start part, included in removal.
+	 * @param End - End part, not included in removal.
+	 */
+	void RemoveRange(const FPart* const Start, const FPart* const End);
+
+private:
+	const TSharedPtr<FData> Data;
+	bool bHasIntersections;
+
+
+	/**
+	 * Compute expand value at which point's normal will intersect edge (with same expansion). Value is stored in Point->AvailableExpandsFar[Edge].
+	 * @param Point - Point for which available expand should be computed.
+	 * @param Edge - Edge for which available expand should be computed.
+	 */
+	void ComputeAvailableExpandFar(FPart* const Point, const FPart* const Edge);
+};
diff --git a/Engine/Plugins/Experimental/Text3D/Source/Text3D/Private/ContourList.cpp b/Engine/Plugins/Experimental/Text3D/Source/Text3D/Private/ContourList.cpp
new file mode 100644
--- /dev/null
+++ b/Engine/Plugins/Experimental/Text3D/Source/Text3D/Private/ContourList.cpp
@@ -0,0 +1,80 @@
+﻿// Copyright 1998-2019 Epic Games, Inc. All Rights Reserved.
+
+
+#include "ContourList.h"
+#include "Data.h"
+#include "Part.h"
+#include "Intersection.h"
+#include "GlyphLoader.h"
+
+#include "Math/UnrealMathUtility.h"
+
+
+FContourList::FContourList(const FT_GlyphSlot Glyph, const TSharedPtr<FData> DataIn)
+	: Data(DataIn)
+{
+	FGlyphLoader().Load(Glyph, Data, this);
+	Init();
+}
+
+FContour& FContourList::Add()
+{
+	AddTail(FContour(Data));
+	return GetTail()->GetValue();
+}
+
+void FContourList::Remove(const FContour& Contour)
+{
+	// Search with comparing pointers
+	for (TDoubleLinkedList<FContour>::TDoubleLinkedListNode* Node = GetHead(); Node; Node = Node->GetNextNode())
+	{
+		if (&Node->GetValue() == &Contour)
+		{
+			RemoveNode(Node);
+			break;
+		}
+	}
+}
+
+void FContourList::Reset()
+{
+	for (FContour& Contour : *this)
+	{
+		for (FPart* const Part : Contour)
+		{
+			Part->ResetDoneExpand();
+			Part->ResetInitialPosition();
+		}
+
+		Contour.ResetContour();
+	}
+}
+
+void FContourList::Init()
+{
+	for (FContour& Contour : *this)
+	{
+		for (FPart* const Point : Contour)
+		{
+			if (!Point->bSmooth)
+			{
+				Point->ComputeSmooth();
+			}
+
+			Point->ResetInitialPosition();
+		}
+
+		for (FPart* const Point : Contour)
+		{
+			Contour.ComputeAvailableExpandNear(Point);
+		}
+
+		if (Contour.Num() > FIntersection::MinContourSizeForIntersectionFar)
+		{
+			for (FPart* const Point : Contour)
+			{
+				Contour.ComputeAvailableExpandsFarFrom(Point);
+			}
+		}
+	}
+}
diff --git a/Engine/Plugins/Experimental/Text3D/Source/Text3D/Private/ContourList.h b/Engine/Plugins/Experimental/Text3D/Source/Text3D/Private/ContourList.h
new file mode 100644
--- /dev/null
+++ b/Engine/Plugins/Experimental/Text3D/Source/Text3D/Private/ContourList.h
@@ -0,0 +1,36 @@
+// Copyright 1998-2019 Epic Games, Inc. All Rights Reserved.
+
+#pragma once
+
+
+#include "Text3DPrivate.h"
+#include "Contour.h"
+#include "Containers/List.h"
+
+
+class FData;
+
+class FContourList final : public TDoubleLinkedList<FContour>
+{
+public:
+	FContourList(const FT_GlyphSlot Glyph, const TSharedPtr<FData> DataIn);
+
+	/**
+	 * Create contour.
+	 * @return Reference to created contour.
+	 */
+	FContour& Add();
+	/**
+	 * Remove contour.
+	 * @param Contour - Const reference to contour that should be removed.
+	 */
+	void Remove(const FContour& Contour);
+
+	void Reset();
+
+private:
+	const TSharedPtr<FData> Data;
+
+
+	void Init();
+};
diff --git a/Engine/Plugins/Experimental/Text3D/Source/Text3D/Private/Data.cpp b/Engine/Plugins/Experimental/Text3D/Source/Text3D/Private/Data.cpp
--- a/Engine/Plugins/Experimental/Text3D/Source/Text3D/Private/Data.cpp
+++ b/Engine/Plugins/Experimental/Text3D/Source/Text3D/Private/Data.cpp
@@ -2,7 +2,7 @@
 
 
 #include "Data.h"
-#include "Bevel/Part.h"
+#include "Part.h"
 
 
 FData::FData(TSharedRef<TText3DMeshList> MeshesIn, const float ExpandTotalIn, const float FontInverseScaleIn, const FVector& ScaleIn)
@@ -129,11 +129,6 @@ float FData::GetFontInverseScale() const
 	return FontInverseScale;
 }
 
-float FData::GetExpandTarget() const
-{
-	return ExpandTarget;
-}
-
 void FData::ResetDoneExtrude()
 {
 	DoneExtrude = 0;
@@ -154,7 +149,7 @@ FVector FData::ComputeTangentZ(const FPart* const Edge, const float DoneExpand)
 {
 	const FVector2D TangentX = Edge->TangentX;
 
-	const float t = Expand == 0 ? 0 : DoneExpand / Expand;
+	const float t = FMath::IsNearlyZero(Expand) ? 0.0f : DoneExpand / Expand;
 	const FVector2D Normal = NormalStart * (1 - t) + NormalEnd * t;
 
 	return FVector(FVector2D(TangentX.Y, -TangentX.X) * Normal.X, Normal.Y);
@@ -168,7 +163,124 @@ void FData::SetCurrentMesh(EText3DMeshType Type)
 	CurrentMesh->GlyphStartVertices.Add(CurrentMesh->Vertices.Num());
 }
 
+FVector2D FData::Expanded(const FPart* const Point) const
+{
+	// Needed expand value is difference of total expand and point's done expand
+	return Point->Expanded(ExpandTarget - Point->DoneExpand);
+}
+
+void FData::ExpandPoint(FPart* const Point, const int32 Count)
+{
+	Point->Position = Expanded(Point);
+
+
+	// Find first previous point that expands to other position
+	const FPart* const Curr = Point;
+
+	FPart* Prev = Point;
+	for (int32 Index = 0; Index < Count - 1; Index++)
+	{
+		Prev = Prev->Prev;
+	}
+
+	// Find first next point that expands to other position
+	const FPart* const Next = Point->Next;
+
+
+	int32 v = AddVertices(1);
+	FPart* p = Prev->Next;
+
+	auto PushNext = [&p, &v]()
+	{
+		// If point is smooth, only one vertex is needed
+		if (!p->bSmooth)
+			v++;
+
+		p->PathNext.Add(v);
+	};
+
+	// Write indices to paths before creating them
+	p->PathPrev.Add(v);
+
+	for (; p != Curr; p = p->Next)
+	{
+		PushNext();
+		p->Next->PathPrev.Add(v);
+	}
+
+	PushNext();
+
+
+	FVector2D TangentX = Prev->TangentX;
+	FVector TangentZ = ComputeTangentZ(Prev, Point->DoneExpand);
+
+	for (p = Prev->Next; p != Next; p = p->Next)
+	{
+		if (p->bSmooth)
+		{
+			TangentX += p->TangentX;
+			TangentZ += ComputeTangentZ(p, Point->DoneExpand);
+		}
+		else
+		{
+			AddVertex(Point, TangentX.GetSafeNormal(), TangentZ.GetSafeNormal());
+			AddVertices(1);
+
+			TangentX = p->TangentX;
+			TangentZ = ComputeTangentZ(p, Point->DoneExpand);
+		}
+	}
+
+	AddVertex(Point, TangentX.GetSafeNormal(), TangentZ.GetSafeNormal());
+}
+
+void FData::FillEdge(FPart* const Edge, const bool bSkipLastTriangle)
+{
+	FPart* const EdgeA = Edge;
+	FPart* const EdgeB = Edge->Next;
+
+	MakeTriangleFanAlongNormal(EdgeB, EdgeA, false, true);
+	MakeTriangleFanAlongNormal(EdgeA, EdgeB, true, false);
+
+	if (!bSkipLastTriangle)
+	{
+		MakeTriangleFanAlongNormal(EdgeB, EdgeA, false, false);
+	}
+	else
+	{
+		// Index has to be removed despite last triangle being skipped.
+		// For example when normals intersect and result of expansion of EdgeA and EdgeB is one point -
+		// this point was already covered with last MakeTriangleFanAlongNormal call,
+		// no need to keep it in neighbour point's path.
+		EdgeA->PathNext.RemoveAt(0);
+	}
+
+	// Write done expand
+	EdgeA->DoneExpand = ExpandTarget;
+	EdgeB->DoneExpand = ExpandTarget;
+}
+
 FVector FData::GetVector(const FVector2D Position, const float Height) const
 {
 	return (FVector(0, Position.X, Position.Y) * FontInverseScale + FVector(Height, HorizontalOffset, VerticalOffset)) * Scale;
 }
+
+void FData::MakeTriangleFanAlongNormal(const FPart* const Cap, FPart* const Normal, const bool bNormalIsCapNext, const bool bSkipLastTriangle)
+{
+	TArray<int32>& Path = bNormalIsCapNext ? Normal->PathPrev : Normal->PathNext;
+	const int32 Count = Path.Num() - (bSkipLastTriangle ? 2 : 1);
+
+	// Create triangles
+	AddTriangles(Count);
+
+	for (int32 Index = 0; Index < Count; Index++)
+	{
+		AddTriangle((bNormalIsCapNext ? Cap->PathNext : Cap->PathPrev)[0],
+				Path[bNormalIsCapNext ? Index + 1 : Index],
+				Path[bNormalIsCapNext ? Index : Index + 1]
+			);
+	}
+
+	// Remove covered vertices from path
+	Path.RemoveAt(0, Count);
+}
diff --git a/Engine/Plugins/Experimental/Text3D/Source/Text3D/Private/Data.h b/Engine/Plugins/Experimental/Text3D/Source/Text3D/Private/Data.h
--- a/Engine/Plugins/Experimental/Text3D/Source/Text3D/Private/Data.h
+++ b/Engine/Plugins/Experimental/Text3D/Source/Text3D/Private/Data.h
@@ -3,7 +3,7 @@
 #pragma once
 
 
-#include "Bevel/Mesh.h"
+#include "Mesh.h"
 #include "CoreMinimal.h"
 #include "DynamicMeshBuilder.h"
 
@@ -18,8 +18,8 @@ public:
 	 * Contructor.
 	 * @param MeshesIn - Vertices and indices to which bevel should be added (contains front cap).
 	 * @param ExpandTotalIn - Total expand value (for all ArcSegments from Bevel.cpp).
-	 * @param FontInverseScaleIn - Documented in Bevel.h.
-	 * @param ScaleIn - Documented in Bevel.h.
+	 * @param FontInverseScaleIn - Documented in ContourList.h.
+	 * @param ScaleIn - Documented in ContourList.h.
 	 */
 	FData(TSharedRef<TText3DMeshList> MeshesIn, const float ExpandTotalIn, const float FontInverseScaleIn, const FVector& ScaleIn);
 
@@ -53,8 +53,6 @@ public:
 
 	float GetFontInverseScale() const;
 
-	float GetExpandTarget() const;
-
 	void ResetDoneExtrude();
 	void IncreaseDoneExtrude();
 
@@ -62,6 +60,26 @@ public:
 	FVector ComputeTangentZ(const FPart* const Edge, const float DoneExpand);
 	void SetCurrentMesh(EText3DMeshType Type);
 
+
+	/**
+	 * FPart::Expanded for total expand value Data::ExpandTarget.
+	 * @param Point - Point for which position should be computed.
+	 * @return Computed position.
+	 */
+	FVector2D Expanded(const FPart* const Point) const;
+	/**
+	 * Similar to FData::Expanded but actually creates vertices and writes indices to paths.
+	 * @param Point - Point that should be expanded.
+	 * @param Count - Amount of edges to which result point will belong. 2 for case without intersections, (n + 1) for case when (n) normals intersect in result point.
+	 */
+	void ExpandPoint(FPart* const Point, const int32 Count);
+	/**
+	 * Make triangulation of edge along paths of it's vertices (from end of previous triangulation to result of points' expansion). Removes covered points' indices from paths.
+	 * @param Edge - Edge that has to be filled.
+	 * @param bSkipLastTriangle - Do not create last triangle (furthest from end of previous triangulation).
+	 */
+	void FillEdge(FPart* const Edge, const bool bSkipLastTriangle);
+
 private:
 	TSharedRef<TText3DMeshList> Meshes;
 	FText3DDynamicData* CurrentMesh;
@@ -98,4 +116,12 @@ private:
 	 * @return 3d coordinate.
 	 */
 	FVector GetVector(const FVector2D Position, const float Height) const;
+	/**
+	 * Make triangle fan, called from FData::FillEdge.
+	 * @param Cap - Cap of triangle fan.
+	 * @param Normal - Point, fan will be created along it's normal.
+	 * @param bNormalIsCapNext - Normal is next point after cap or vice versa.
+	 * @param bSkipLastTriangle - See FData::FillEdge.
+	 */
+	void MakeTriangleFanAlongNormal(const FPart* const Cap, FPart* const Normal, const bool bNormalIsCapNext, const bool bSkipLastTriangle);
 };
diff --git a/Engine/Plugins/Experimental/Text3D/Source/Text3D/Private/GlyphLoader.cpp b/Engine/Plugins/Experimental/Text3D/Source/Text3D/Private/GlyphLoader.cpp
new file mode 100644
--- /dev/null
+++ b/Engine/Plugins/Experimental/Text3D/Source/Text3D/Private/GlyphLoader.cpp
@@ -0,0 +1,613 @@
+// Copyright 1998-2019 Epic Games, Inc. All Rights Reserved.
+
+#include "GlyphLoader.h"
+#include "Contour.h"
+#include "Part.h"
+#include "ContourList.h"
+
+FGlyphLoader::FGlyphLoader()
+	: EndIndex(-1)
+{
+
+}
+
+struct FGlyphLoader::FContourNode
+{
+	FContourNode(FContour* const ContourIn)
+		: Contour(ContourIn)
+	{
+
+	}
+
+
+	FContour* Contour;
+	TArray<FContourNode*> Nodes;
+};
+
+void FGlyphLoader::Load(const FT_GlyphSlot Glyph, const TSharedPtr<FData> Data, FContourList* Contours)
+{
+	const FT_Outline Outline = Glyph->outline;
+	const int32 ContourCount = Outline.n_contours;
+	Clockwise.Reserve(ContourCount);
+	FContourNode Root(nullptr);
+
+	for (int32 Index = 0; Index < ContourCount; Index++)
+	{
+		if (CreateContour(Outline, Index, Contours))
+		{
+			ComputeInitialParity();
+			Insert(new FContourNode(Contour), &Root);
+		}
+	}
+
+	FixParity(&Root, false);
+}
+
+
+FGlyphLoader::FLine::FLine(const FVector2D PositionIn)
+	: Position(PositionIn)
+{
+
+}
+
+void FGlyphLoader::FLine::Add(FGlyphLoader* const Loader)
+{
+	if (!Loader->Contour->Num() || !(Position - Loader->FirstPosition).IsNearlyZero())
+	{
+		FPart* const Point = Loader->AddPoint(Position);
+		Loader->JoinWithLast(Point);
+		Loader->LastPoint = Point;
+	}
+}
+
+
+FGlyphLoader::FCurve::FCurve(const bool bLineIn)
+	: bLine(bLineIn)
+
+	, StartT(0.f)
+	, EndT(1.f)
+{
+
+}
+
+FGlyphLoader::FCurve::~FCurve()
+{
+
+}
+
+struct FGlyphLoader::FCurve::FPointData
+{
+	float T;
+	FVector2D Position;
+	FVector2D Tangent;
+	FPart* Point;
+};
+
+void FGlyphLoader::FCurve::Add(FGlyphLoader* const LoaderIn)
+{
+	Loader = LoaderIn;
+
+	if (bLine)
+	{
+		FLine(Position(StartT)).Add(Loader);
+		return;
+	}
+
+	Depth = 0;
+
+
+	FPointData Start;
+	Start.T = StartT;
+	Start.Position = Position(Start.T);
+	Start.Tangent = Tangent(Start.T);
+	Start.Point = Loader->AddPoint(Start.Position);
+	First = Start.Point;
+	bFirstSplit = false;
+
+	FPointData End;
+	End.T = EndT;
+	End.Position = Position(End.T);
+	End.Tangent = Tangent(End.T);
+	End.Point = nullptr;
+	Last = End.Point;
+	bLastSplit = false;
+
+
+	Loader->JoinWithLast(Start.Point);
+	ComputeMaxDepth();
+	Split(Start, End);
+}
+
+bool FGlyphLoader::FCurve::OnOneLine(const FVector2D A, const FVector2D B, const FVector2D C)
+{
+	return FMath::IsNearlyZero(FVector2D::CrossProduct((B - A).GetSafeNormal(), (C - A).GetSafeNormal()));
+}
+
+void FGlyphLoader::FCurve::ComputeMaxDepth()
+{
+	// Compute approximate curve length with 4 points
+	const float MinStep = 30.f;
+	const float StepT = 0.333f;
+	float Length = 0.f;
+
+	FVector2D Prev;
+	FVector2D Curr = Position(StartT);
+
+	for (float T = StartT + StepT; T < EndT; T += StepT)
+	{
+		Prev = Curr;
+		Curr = Position(T);
+
+		Length += (Curr - Prev).Size();
+	}
+
+	const float MaxStepCount = Length / MinStep;
+	MaxDepth = static_cast<int32>(FMath::Log2(MaxStepCount)) + 1;
+}
+
+void FGlyphLoader::FCurve::Split(const FPointData& Start, const FPointData& End)
+{
+	Depth++;
+	FPointData Middle;
+
+	Middle.T = (Start.T + End.T) / 2.f;
+	Middle.Position = Position(Middle.T);
+	Middle.Tangent = Tangent(Middle.T);
+	Middle.Point = Loader->AddPoint(Middle.Position);
+	Middle.Point->bSmooth = true;
+
+
+	Start.Point->Next = Middle.Point;
+	Middle.Point->Prev = Start.Point;
+	Middle.Point->Next = End.Point;
+
+	if (End.Point)
+	{
+		End.Point->Prev = Middle.Point;
+	}
+	else
+	{
+		Loader->LastPoint = Middle.Point;
+	}
+
+
+	CheckPart(Start, Middle);
+	UpdateTangent(&Middle);
+	CheckPart(Middle, End);
+	Depth--;
+}
+
+void FGlyphLoader::FCurve::CheckPart(const FPointData& Start, const FPointData& End)
+{
+	const FVector2D Side = (End.Position - Start.Position).GetSafeNormal();
+
+	if ((FVector2D::DotProduct(Side, Start.Tangent) > FPart::CosMaxAngleSideTangent && FVector2D::DotProduct(Side, End.Tangent) > FPart::CosMaxAngleSideTangent) || Depth >= MaxDepth)
+	{
+		if (!bFirstSplit && Start.Point == First)
+		{
+			bFirstSplit = true;
+			Split(Start, End);
+		}
+		else if (!bLastSplit && End.Point == Last)
+		{
+			bLastSplit = true;
+			Split(Start, End);
+		}
+		else
+		{
+			Start.Point->TangentX = Side;
+		}
+	}
+	else
+	{
+		Split(Start, End);
+	}
+}
+
+void FGlyphLoader::FCurve::UpdateTangent(FPointData* const Middle)
+{
+
+}
+
+
+FGlyphLoader::FQuadraticCurve::FQuadraticCurve(const FVector2D A, const FVector2D B, const FVector2D C)
+	: FCurve(OnOneLine(A, B, C))
+
+	, E(A - 2 * B + C)
+	, F(-A + B)
+	, G(A)
+{
+
+}
+
+FVector2D FGlyphLoader::FQuadraticCurve::Position(const float T) const
+{
+	return E * T * T + 2 * F * T + G;
+}
+
+FVector2D QuadraticCurveTangent(const FVector2D E, const FVector2D F, const float T)
+{
+	const FVector2D Result = E * T + F;
+
+	if (Result.IsNearlyZero())
+	{
+		// Just some vector with non-zero length
+		return {1.f, 0.f};
+	}
+
+	return Result;
+}
+
+FVector2D FGlyphLoader::FQuadraticCurve::Tangent(const float T)
+{
+	return QuadraticCurveTangent(E, F, T).GetSafeNormal();
+}
+
+
+FGlyphLoader::FCubicCurve::FCubicCurve(const FVector2D A, const FVector2D B, const FVector2D C, const FVector2D D)
+	: FCurve(OnOneLine(A, B, C) && OnOneLine(B, C, D))
+
+	, E(-A + 3 * B - 3 * C + D)
+	, F(A - 2 * B + C)
+	, G(-A + B)
+	, H(A)
+{
+	if (!bLine)
+	{
+		bSharpStart = G.IsNearlyZero();
+		bSharpEnd = (C - D).IsNearlyZero();
+	}
+}
+
+void FGlyphLoader::FCubicCurve::UpdateTangent(FPointData* const Middle)
+{
+	// In this point curve is not smooth, and  r'(t + 0) / |r'(t + 0)| = -r'(t - 0) / |r'(t - 0)|
+	if (bSharpMiddle)
+	{
+		bSharpMiddle = false;
+		Middle->Tangent *= -1;
+		Middle->Point->bSmooth = false;
+	}
+}
+
+FVector2D FGlyphLoader::FCubicCurve::Position(const float T) const
+{
+	return E * T * T * T + 3 * F * T * T + 3 * G * T + H;
+}
+
+FVector2D FGlyphLoader::FCubicCurve::Tangent(const float T)
+{
+	FVector2D Result;
+
+	// Using  r' / |r'|  for sharp start and end
+	if (bSharpStart && FMath::IsNearlyEqual(T, StartT))
+	{
+		Result = F;
+	}
+	else if (bSharpEnd && FMath::IsNearlyEqual(T, EndT))
+	{
+		Result = -(E + F);
+	}
+	else
+	{
+		Result = E * T * T + 2 * F * T + G;
+		bSharpMiddle = Result.IsNearlyZero();
+
+		if (bSharpMiddle)
+		{
+			// Using derivative of quadratic bezier curve (A, B, C) in this point
+			Result = QuadraticCurveTangent(F, G, T);
+		}
+	}
+
+	return Result.GetSafeNormal();
+}
+
+
+bool FGlyphLoader::CreateContour(const FT_Outline Outline, const int32 ContourIndex, FContourList* Contours)
+{
+	const int32 StartIndex = EndIndex + 1;
+	EndIndex = Outline.contours[ContourIndex];
+	const int32 ContourLength = EndIndex - StartIndex + 1;
+
+	if (ContourLength < 3)
+	{
+		return false;
+	}
+
+	Contour = &Contours->Add();
+
+
+	const FT_Vector* const Points = Outline.points + StartIndex;
+	auto ToFVector2D = [Points](const int32 Index)
+	{
+		const FT_Vector Point = Points[Index];
+		return FVector2D(Point.x, Point.y);
+	};
+
+	FVector2D Prev;
+	FVector2D Curr = ToFVector2D(ContourLength - 1);
+	FVector2D Next = ToFVector2D(0);
+	FVector2D NextNext = ToFVector2D(1);
+
+
+	const char* const Tags = Outline.tags + StartIndex;
+	auto Tag = [Tags](int32 Index)
+	{
+		return FT_CURVE_TAG(Tags[Index]);
+	};
+
+	int32 TagPrev;
+	int32 TagCurr = Tag(ContourLength - 1);
+	int32 TagNext = Tag(0);
+
+
+	FVector2D& FirstPositionLocal = FirstPosition;
+	FPart* const& LastPointLocal = LastPoint;
+	FContour*& ContourLocal = Contour;
+	auto ContourIsBad = [&FirstPositionLocal, &LastPointLocal, &ContourLocal](const FVector2D Point)
+	{
+		if (ContourLocal->Num() == 0)
+		{
+			FirstPositionLocal = Point;
+			return false;
+		}
+
+		return (Point - LastPointLocal->Position).IsNearlyZero();
+	};
+
+	auto RemoveContour = [Contours, &ContourLocal]()
+	{
+		Contours->Remove(*ContourLocal);
+		ContourLocal = nullptr;
+	};
+
+
+	for (int32 Index = 0; Index < ContourLength; Index++)
+	{
+		const int32 NextIndex = (Index + 1) % ContourLength;
+
+		Prev = Curr;
+		Curr = Next;
+		Next = NextNext;
+		NextNext = ToFVector2D((NextIndex + 1) % ContourLength);
+
+		TagPrev = TagCurr;
+		TagCurr = TagNext;
+		TagNext = Tag(NextIndex);
+
+		if (TagCurr == FT_Curve_Tag_On)
+		{
+			if (TagNext == FT_Curve_Tag_Cubic || TagNext == FT_Curve_Tag_Conic)
+			{
+				continue;
+			}
+
+			if (ContourIsBad(Curr))
+			{
+				RemoveContour();
+				return false;
+			}
+
+			if (TagNext == FT_Curve_Tag_On && (Curr - Next).IsNearlyZero())
+			{
+				continue;
+			}
+
+			FLine(Curr).Add(this);
+		}
+		else if (TagCurr == FT_Curve_Tag_Conic)
+		{
+			FVector2D A;
+
+			if (TagPrev == FT_Curve_Tag_On)
+			{
+				if (ContourIsBad(Prev))
+				{
+					RemoveContour();
+					return false;
+				}
+
+				A = Prev;
+			}
+			else
+			{
+				A = (Prev + Curr) / 2.f;
+			}
+
+			FQuadraticCurve(A, Curr, TagNext == FT_Curve_Tag_Conic ? (Curr + Next) / 2.f : Next).Add(this);
+		}
+		else if (TagCurr == FT_Curve_Tag_Cubic && TagNext == FT_Curve_Tag_Cubic)
+		{
+			if (ContourIsBad(Prev))
+			{
+				RemoveContour();
+				return false;
+			}
+
+			FCubicCurve(Prev, Curr, Next, NextNext).Add(this);
+		}
+	}
+
+	if (Contour->Num() < 3)
+	{
+		RemoveContour();
+		return false;
+	}
+
+	JoinWithLast((*Contour)[0]);
+
+	for (FPart* const Edge : *Contour)
+	{
+		if (!Edge->ComputeNormal())
+		{
+			RemoveContour();
+			return false;
+		}
+	}
+
+	return true;
+}
+
+void FGlyphLoader::ComputeInitialParity()
+{
+	float DoubledArea = 0.0f;
+
+	FVector2D Curr;
+	FVector2D Next = FirstPosition;
+
+	Next = (*Contour)[1]->Position - FirstPosition;
+
+	for (FPart* Point = (*Contour)[2]; Point != (*Contour)[0]; Point = Point->Next)
+	{
+		Curr = Next;
+		Next = Point->Position - FirstPosition;
+
+		DoubledArea += Curr ^ Next;
+	}
+
+	Clockwise.Add(TPair<const FContour*, bool>(Contour, DoubledArea < 0));
+}
+
+void FGlyphLoader::Insert(FContourNode* const NodeA, FContourNode* const NodeB)
+{
+	TArray<FContourNode*>& BNodes = NodeB->Nodes;
+	const FContour* ContourA = NodeA->Contour;
+
+	for (int32 IndexC = 0; IndexC < BNodes.Num(); IndexC++)
+	{
+		FContourNode*& NodeC = BNodes[IndexC];
+		const FContour* ContourC = NodeC->Contour;
+
+		if (Inside(ContourA, ContourC))
+		{
+			Insert(NodeA, NodeC);
+			return;
+		}
+
+		if (Inside(ContourC, ContourA))
+		{
+			// add contourC to list of contours that are inside contourA
+			TArray<FContourNode*>& ANodes = NodeA->Nodes;
+			ANodes.Add(NodeC);
+			// replace contourC with contourA in list it was before
+			NodeC = NodeA;
+
+			// check if other contours in that list are inside contourA
+
+			for (int32 Index = BNodes.Num() - 1; Index > IndexC; Index--)
+			{
+				FContourNode* Node = BNodes[Index];
+
+				if (Inside(Node->Contour, ContourA))
+				{
+					ANodes.Add(Node);
+					BNodes.RemoveAt(Index);
+				}
+			}
+
+			return;
+		}
+	}
+
+	BNodes.Add(NodeA);
+}
+
+void FGlyphLoader::FixParity(FContourNode* const Node, const bool bClockwiseIn)
+{
+	for (FContourNode* NodeA : Node->Nodes)
+	{
+		FixParity(NodeA, !bClockwiseIn);
+		FContour* const ContourB = NodeA->Contour;
+		delete NodeA;
+
+		if (bClockwiseIn != Clockwise[ContourB])
+		{
+			for (FPart* const Point : *ContourB)
+			{
+				Swap(Point->Prev, Point->Next);
+			}
+
+
+			FPart* const First = (*ContourB)[0];
+			FPart* const Last = First->Prev;
+
+			const FVector2D TangentFirst = First->TangentX;
+
+			for (FPart* Edge = First; Edge != Last; Edge = Edge->Next)
+			{
+				Edge->TangentX = -Edge->Next->TangentX;
+			}
+
+			Last->TangentX = -TangentFirst;
+
+
+			for (FPart* const Point : *ContourB)
+			{
+				Point->Normal *= -1.0f;
+			}
+		}
+	}
+}
+
+FPart* FGlyphLoader::AddPoint(const FVector2D Position)
+{
+	FPart* const Point = new FPart();
+	Contour->Add(Point);
+	Point->Position = Position;
+	return Point;
+}
+
+void FGlyphLoader::JoinWithLast(FPart* const Point)
+{
+	if (Contour->Num() > 1)
+	{
+		LastPoint->Next = Point;
+		Point->Prev = LastPoint;
+
+		LastPoint->ComputeTangentX();
+	}
+}
+
+bool FGlyphLoader::Inside(const FContour* const ContourA, const FContour* const ContourB)
+{
+	const int32 BPointCount = ContourB->Num();
+	// compute angle to which vector "b(i) - a(0)" is rotated when "i" iterates over all points of contourB
+	float AngleTotal = 0;
+
+	float AnglePrev;
+	float AngleCurr;
+
+	const TMap<const FContour*, bool>& ClockwiseLocal = Clockwise;
+	auto ComputeAngleCurr = [ContourB, &ClockwiseLocal, BPointCount, ContourA, &AngleCurr](const int32 EndPoint)
+	{
+		// use counterclockwise version of contour if it's clockwise
+		const FVector2D Delta = (*ContourB)[ClockwiseLocal[ContourB] ? BPointCount - 1 - EndPoint : EndPoint]->Position - (*ContourA)[0]->Position;
+		AngleCurr = FMath::Atan2(Delta.Y, Delta.X);
+	};
+
+	ComputeAngleCurr(0);
+
+	for (int32 Index = 0; Index < BPointCount; Index++)
+	{
+		AnglePrev = AngleCurr;
+		ComputeAngleCurr((Index + 1) % BPointCount);
+
+		float DeltaAngle = AngleCurr - AnglePrev;
+
+		if (DeltaAngle < -PI)
+		{
+			DeltaAngle += 2 * PI;
+		}
+
+		if (DeltaAngle > PI)
+		{
+			DeltaAngle -= 2 * PI;
+		}
+
+		AngleTotal += DeltaAngle;
+	}
+
+	// if contourA is inside contourB, angle is 2pi, else it's 0
+	return AngleTotal > 3;
+}
diff --git a/Engine/Plugins/Experimental/Text3D/Source/Text3D/Private/GlyphLoader.h b/Engine/Plugins/Experimental/Text3D/Source/Text3D/Private/GlyphLoader.h
new file mode 100644
--- /dev/null
+++ b/Engine/Plugins/Experimental/Text3D/Source/Text3D/Private/GlyphLoader.h
@@ -0,0 +1,158 @@
+// Copyright 1998-2019 Epic Games, Inc. All Rights Reserved.
+
+#pragma once
+
+#include "Text3DPrivate.h"
+#include "Templates/SharedPointer.h"
+#include "Math/Vector2D.h"
+
+class FData;
+class FContourList;
+struct FPart;
+class FContour;
+
+class FGlyphLoader final
+{
+public:
+	FGlyphLoader();
+	void Load(const FT_GlyphSlot Glyph, const TSharedPtr<FData> Data, FContourList* Contours);
+
+private:
+	struct FContourNode;
+
+	class FLine final
+	{
+	public:
+		FLine(const FVector2D PositionIn);
+		void Add(FGlyphLoader* const Loader);
+
+	private:
+		const FVector2D Position;
+	};
+
+	class FCurve
+	{
+	public:
+		FCurve(const bool bLineIn);
+		virtual ~FCurve();
+
+		void Add(FGlyphLoader* const LoaderIn);
+
+	protected:
+		struct FPointData;
+
+
+		/**
+		 * Check if 3 points are on one line.
+		 * @param A - Point A.
+		 * @param B - Point B.
+		 * @param C - Point C.
+		 * @return Are points on one line?
+		 */
+		static bool OnOneLine(const FVector2D A, const FVector2D B, const FVector2D C);
+
+
+		/** Is curve a line? */
+		const bool bLine;
+		FGlyphLoader* Loader;
+
+		const float StartT;
+		const float EndT;
+
+	private:
+		int32 Depth;
+		int32 MaxDepth;
+
+		const FPart* First;
+		const FPart* Last;
+
+		/** Needed to make additional splits near start and end of curve */
+		bool bFirstSplit;
+		bool bLastSplit;
+
+
+		/**
+		 * Compute max depth (depends on curve length, step is fixed).
+		 */
+		void ComputeMaxDepth();
+		void Split(const FPointData& Start, const FPointData& End);
+		void CheckPart(const FPointData& Start, const FPointData& End);
+		virtual void UpdateTangent(FPointData* const Middle);
+
+		virtual FVector2D Position(const float T) const = 0;
+		virtual FVector2D Tangent(const float T) = 0;
+	};
+
+	class FQuadraticCurve final : public FCurve
+	{
+	public:
+		FQuadraticCurve(const FVector2D A, const FVector2D B, const FVector2D C);
+
+	private:
+		const FVector2D E;
+		const FVector2D F;
+		const FVector2D G;
+
+
+		FVector2D Position(const float T) const override;
+		FVector2D Tangent(const float T) override;
+	};
+
+	class FCubicCurve final : public FCurve
+	{
+	public:
+		FCubicCurve(const FVector2D A, const FVector2D B, const FVector2D C, const FVector2D D);
+
+	private:
+		const FVector2D E;
+		const FVector2D F;
+		const FVector2D G;
+		const FVector2D H;
+
+		/** Sharp means that curve derivative has zero length, it's actually sharp only in middle case */
+		bool bSharpStart;
+		bool bSharpMiddle;
+		bool bSharpEnd;
+
+
+		void UpdateTangent(FPointData* const Middle) override;
+
+		FVector2D Position(const float t) const override;
+		FVector2D Tangent(const float t) override;
+	};
+
+
+	int32 EndIndex;
+	FContour* Contour;
+	/** Initial parity */
+	TMap<const FContour*, bool> Clockwise;
+
+	FVector2D FirstPosition;
+	FPart* LastPoint;
+
+
+	bool CreateContour(const FT_Outline Outline, const int32 ContourIndex, FContourList* Contours);
+	void ComputeInitialParity();
+	/**
+	 * Insert NodeA inside NodeB.
+	 * @param NodeA
+	 * @param NodeB
+	 */
+	void Insert(FContourNode* const NodeA, FContourNode* const NodeB);
+	/**
+	 * Reverse contour if it's initial parity differs from the one it should have.
+	 * @param Node - Function is called recursively to fix all contours inside Node->Contour.
+	 * @param ClockwiseIn - The parity that contours listed in Node->Nodes should have.
+	 */
+	void FixParity(FContourNode* const Node, const bool bClockwiseIn);
+
+	FPart* AddPoint(const FVector2D Position);
+	void JoinWithLast(FPart* const Point);
+	/**
+	 * Check if ContourA is inside ContourB.
+	 * @param ContourA
+	 * @param ContourB
+	 * @return Is ContourA inside ContourB?
+	 */
+	bool Inside(const FContour* const ContourA, const FContour* const ContourB);
+};
