Copy up //UE4/Dev-Enterprise/...at  Changelist 10508152  to //UE4/Dev-Tools-Staging/...

diff --git a/Engine/Plugins/Experimental/Text3D/Source/Text3D/Private/Intersection.cpp b/Engine/Plugins/Experimental/Text3D/Source/Text3D/Private/Intersection.cpp
new file mode 100644
--- /dev/null
+++ b/Engine/Plugins/Experimental/Text3D/Source/Text3D/Private/Intersection.cpp
@@ -0,0 +1,366 @@
+// Copyright 1998-2019 Epic Games, Inc. All Rights Reserved.
+
+
+#include "Intersection.h"
+#include "MeshCreator.h"
+#include "Contour.h"
+#include "Part.h"
+#include "Util.h"
+#include "Data.h"
+#include "ContourList.h"
+
+
+FIntersection::FIntersection(const TSharedPtr<FData> DataIn, const TSharedPtr<FContourList> ContoursIn, FContour* const ContourIn)
+	: Data(DataIn)
+	, Contours(ContoursIn)
+
+	, Contour(*ContourIn)
+
+	, Vertex(nullptr)
+	// Make Value more then max expand to invalidate intersection
+	, Value(Data->GetExpand() * 2)
+{
+
+}
+
+FIntersection::~FIntersection()
+{
+
+}
+
+float FIntersection::GetValue() const
+{
+	return Value;
+}
+
+FPart* FIntersection::GetVertex() const
+{
+	return Vertex;
+}
+
+bool FIntersection::ContourHasCloserIntersectionAt(FPart* const Point, const float Expand)
+{
+	const float AvailableExpand = Point->DoneExpand + Expand;
+
+	if (AvailableExpand >= Value)
+	{
+		return false;
+	}
+
+	Value = AvailableExpand;
+	Vertex = Point;
+
+	return true;
+}
+
+
+FIntersectionNear::FIntersectionNear(const TSharedPtr<FData> DataIn, const TSharedPtr<FContourList> ContoursIn, FContour* const ContourIn)
+	: FIntersection(DataIn, ContoursIn, ContourIn)
+{
+	for (FPart* Point : Contour)
+	{
+		ContourHasCloserIntersectionAt(Point, Point->AvailableExpandNear);
+	}
+}
+
+void FIntersectionNear::BevelTillThis()
+{
+	FPart* Curr = GetVertex();
+	const FVector2D Intersection = Data->Expanded(Curr);
+	int32 Count = 1;
+
+	const TSharedPtr<FData> DataLocal = Data;
+	auto ExpandsToSamePoint = [DataLocal, Intersection](const FPart* const Point)
+	{
+		return FMath::IsNearlyZero(FVector2D::DistSquared(DataLocal->Expanded(Point), Intersection), 10);
+	};
+
+	// Check previous points till a point that does not expand to same position is found
+	FPart* Prev = Curr->Prev;
+	for (; ExpandsToSamePoint(Prev); Prev = Prev->Prev, Count++)
+	{
+		if (Prev == Curr)
+		{
+			Contours->Remove(Contour);
+			return;
+		}
+	}
+
+	// Check next points
+	Count += 2;
+	FPart* Next = Curr->Next->Next;
+	for (; ExpandsToSamePoint(Next); Next = Next->Next, Count++)
+	{
+
+	}
+
+	// Curr is the last one of points that expand to same position
+	Curr = Next->Prev;
+
+
+	// Create vertices
+	Data->ExpandPoint(Prev, 2);
+	Data->ExpandPoint(Curr, Count);
+	Data->ExpandPoint(Next, 2);
+
+	const float PrevDelta = GetValue() - Prev->DoneExpand;
+	const float NextDelta = GetValue() - Next->DoneExpand;
+
+
+	// Create triangles
+	Data->FillEdge(Prev, false);
+
+	for (FPart* Edge = Prev->Next; Edge != Curr; Edge = Edge->Next)
+	{
+		Data->FillEdge(Edge, true);
+	}
+
+	Data->FillEdge(Curr, false);
+
+
+	// Stick together parts of contour (it will be broken with removing parts)
+	Curr->PathPrev.Last(0) = Prev->Next->PathPrev.Last(0);
+
+
+	// Remove parts
+	for (FPart* Point = Curr; Point != Prev->Next; Point = Point->Next)
+	{
+		for (const FPart* Edge = Prev->Next; Edge != Curr; Edge = Edge->Next)
+		{
+			Point->AvailableExpandsFar.Remove(Edge);
+		}
+	}
+
+	for (const FPart* Part = Prev->Next->Next; Part != Curr->Next; Part = Part->Next)
+	{
+		delete Part->Prev;
+	}
+
+	Contour.RemoveRange(Prev->Next, Curr);
+
+
+	// Finish sticking together
+	Prev->Next = Curr;
+	Curr->Prev = Prev;
+
+
+	// If vertices were not welded, using simplier implementation
+	const bool bWeldedVertices = Count > 2;
+	if (bWeldedVertices)
+	{
+		Prev->ComputeTangentX();
+		Curr->ComputeTangentX();
+
+		Prev->ComputeNormalAndSmooth();
+		Next->ComputeNormalAndSmooth();
+	}
+
+	Curr->ComputeNormalAndSmooth();
+
+	if (bWeldedVertices)
+	{
+		Prev->ComputeInitialPosition();
+		Next->ComputeInitialPosition();
+	}
+	Curr->ComputeInitialPosition();
+
+	// Update AvailableExpandsNear
+	if (bWeldedVertices)
+	{
+		Contour.ComputeAvailableExpandNear(Prev->Prev);
+	}
+	Contour.ComputeAvailableExpandNear(Prev);
+	Contour.ComputeAvailableExpandNear(Curr);
+
+	if (bWeldedVertices)
+	{
+		Contour.ComputeAvailableExpandNear(Next);
+	}
+	else
+	{
+		Next->AvailableExpandNear -= NextDelta;
+	}
+
+
+	// Update AvailableExpandsFar
+	if (bWeldedVertices)
+	{
+		Contour.ComputeAvailableExpandsFarFrom(Prev);
+		Contour.ComputeAvailableExpandsFarFrom(Next);
+	}
+	else
+	{
+		Prev->AvailableExpandsFar.Remove(Curr);
+		Next->AvailableExpandsFar.Remove(Prev);
+
+		Prev->DecreaseExpandsFar(PrevDelta);
+		Next->DecreaseExpandsFar(NextDelta);
+	}
+	Contour.ComputeAvailableExpandsFarFrom(Curr);
+
+	if (bWeldedVertices)
+	{
+		Contour.ComputeAvailableExpandsFarTo(Prev->Prev);
+		Contour.ComputeAvailableExpandsFarTo(Next);
+	}
+	Contour.ComputeAvailableExpandsFarTo(Prev);
+	Contour.ComputeAvailableExpandsFarTo(Curr);
+}
+
+
+FIntersectionFar::FIntersectionFar(const TSharedPtr<FData> DataIn, const TSharedPtr<FContourList> ContoursIn, FContour* const ContourIn)
+	: FIntersection(DataIn, ContoursIn, ContourIn)
+	, SplitEdge(nullptr)
+{
+	for (FPart* Point : Contour)
+	{
+		for (TPair<FPart *, float>& ExpandFar : Point->AvailableExpandsFar)
+		{
+			if (ContourHasCloserIntersectionAt(Point, ExpandFar.Value))
+			{
+				// Also store edge
+				SplitEdge = ExpandFar.Key;
+			}
+		}
+	}
+}
+
+void FIntersectionFar::BevelTillThis()
+{
+	FPart* const Curr = GetVertex();
+	FPart* const Prev = Curr->Prev;
+	FPart* const Next = Curr->Next;
+
+	// Previous and next points of edge
+	FPart* const EdgeA = SplitEdge;
+	FPart* const EdgeB = EdgeA->Next;
+
+	// Record last index, it will be removed in ExpandPoint and will be needed later
+	const int32 EdgeALast = EdgeA->PathNext.Last(0);
+
+
+	// Create vertices
+	Data->ExpandPoint(Prev, 2);
+	Data->ExpandPoint(Curr, 2);
+	Data->ExpandPoint(Next, 2);
+	Data->ExpandPoint(EdgeA, 2);
+	Data->ExpandPoint(EdgeB, 2);
+
+	const int32 Intersection = Data->AddVertices(1);
+	Data->AddVertex(Curr, EdgeA->TangentX, Data->ComputeTangentZ(EdgeA, Curr->DoneExpand));
+
+
+	const float PrevDelta = GetValue() - Prev->DoneExpand;
+	const float NextDelta = GetValue() - Next->DoneExpand;
+	const float EdgeADelta = GetValue() - EdgeA->DoneExpand;
+	const float EdgeBDelta = GetValue() - EdgeB->DoneExpand;
+
+
+	// Create triangles
+	Data->FillEdge(Prev, false);
+	Data->FillEdge(Curr, false);
+	Data->FillEdge(EdgeA, true);
+
+	Data->AddTriangles(2);
+
+	Data->AddTriangle(EdgeALast, EdgeA->PathNext.Last(0), Intersection);
+	Data->AddTriangle(EdgeALast, Intersection, EdgeB->PathPrev.Last(0));
+
+
+	// Split contour to 2 contours
+	FContour& Initial = Contour;
+	FContour& Added = Contours->Add();
+
+
+	// Copy parts from initial to added contour
+	for (FPart* Part = EdgeB; Part != Curr; Part = Part->Next)
+	{
+		Added.Add(Part);
+	}
+
+	// Make a copy of point that split initial contour, it's needed in added contour
+	FPart* const Copy = new FPart;
+	Added.Add(Copy);
+	// Remove copied parts
+	Initial.RemoveRange(EdgeB, Curr);
+
+	Copy->Position = Curr->Position;
+
+
+	// Stick together initial contour
+	EdgeA->Next = Curr;
+	Curr->Prev = EdgeA;
+
+	// Stick together added contour
+	Prev->Next = Copy;
+	Copy->Prev = Prev;
+	Copy->Next = EdgeB;
+	EdgeB->Prev = Copy;
+
+
+	Copy->DoneExpand = GetValue();
+	Copy->TangentX = EdgeA->TangentX;
+
+	Curr->ComputeNormalAndSmooth();
+	Copy->ComputeNormalAndSmooth();
+
+	Curr->ComputeInitialPosition();
+	Copy->ComputeInitialPosition();
+
+
+	// Update AvailableExpandsNear
+	Initial.ComputeAvailableExpandNear(EdgeA);
+	Initial.ComputeAvailableExpandNear(Curr);
+	Next->AvailableExpandNear -= NextDelta;
+
+	Added.ComputeAvailableExpandNear(Prev);
+	Added.ComputeAvailableExpandNear(Copy);
+	EdgeB->AvailableExpandNear -= EdgeBDelta;
+
+
+	// Finish sticking together added contour
+	Copy->PathPrev.Add(Curr->PathPrev.Last(0));
+	Copy->PathNext.Add(Intersection);
+
+	Curr->PathPrev.Last(0) = Intersection;
+
+
+	// Update AvailableExpandsFar
+	UpdateExpandsFar(&Initial, Added, Curr, EdgeA, EdgeADelta, NextDelta);
+	UpdateExpandsFar(&Added, Initial, Copy, Copy, PrevDelta, EdgeBDelta);
+}
+
+void FIntersectionFar::UpdateExpandsFar(FContour* const UpdatedContour, const FContour& OtherContour, FPart* const Curr, const FPart* const SplitEdgePart, const float PrevDelta, const float NextDelta)
+{
+	// If contour shrinked too much, IntersectionFar is not possible
+	if (UpdatedContour->Num() < MinContourSizeForIntersectionFar)
+	{
+		for (FPart* const Point : *UpdatedContour)
+		{
+			Point->AvailableExpandsFar.Empty();
+		}
+	}
+	else
+	{
+		// Remove references to edges of OtherContour from points of UpdatedContour
+		for (FPart* const Point : *UpdatedContour)
+		{
+			for (const FPart* const Edge : OtherContour)
+			{
+				Point->AvailableExpandsFar.Remove(Edge);
+			}
+		}
+
+		FPart* const Prev = Curr->Prev;
+		FPart* const Next = Curr->Next;
+
+		Prev->AvailableExpandsFar.Remove(Curr);
+		Next->AvailableExpandsFar.Remove(Prev);
+
+		UpdatedContour->ComputeAvailableExpandsFarFrom(Curr);
+		UpdatedContour->ComputeAvailableExpandsFarTo(SplitEdgePart);
+
+		Prev->DecreaseExpandsFar(PrevDelta);
+		Next->DecreaseExpandsFar(NextDelta);
+	}
+}
diff --git a/Engine/Plugins/Experimental/Text3D/Source/Text3D/Private/Intersection.h b/Engine/Plugins/Experimental/Text3D/Source/Text3D/Private/Intersection.h
new file mode 100644
--- /dev/null
+++ b/Engine/Plugins/Experimental/Text3D/Source/Text3D/Private/Intersection.h
@@ -0,0 +1,81 @@
+// Copyright 1998-2019 Epic Games, Inc. All Rights Reserved.
+
+#pragma once
+
+#include "CoreMinimal.h"
+
+class FContour;
+struct FPart;
+class FData;
+class FContourList;
+
+class FIntersection
+{
+public:
+	static const int32 MinContourSizeForIntersectionFar = 4;
+
+
+	FIntersection(const TSharedPtr<FData> DataIn, const TSharedPtr<FContourList> ContoursIn, FContour* const ContourIn);
+	virtual ~FIntersection();
+
+	float GetValue() const;
+
+	virtual void BevelTillThis() = 0;
+
+protected:
+	const TSharedPtr<FData> Data;
+	const TSharedPtr<FContourList> Contours;
+
+	FContour& Contour;
+
+
+	FPart* GetVertex() const;
+	/**
+	 * Check if intersection value for point is less then stored one and store this if it is.
+	 * @param Point - Point that is checked.
+	 * @param Expand - Point's expand value.
+	 * @return Is it less then stored one?
+	 */
+	bool ContourHasCloserIntersectionAt(FPart* const Point, const float Expand);
+
+private:
+	/** Point whose normal intersects. */
+	FPart* Vertex;
+
+	/** Offset in surface of front cap till point of intersection. */
+	float Value;
+};
+
+/** Intersection of point's normal with next point's normal. */
+class FIntersectionNear final : public FIntersection
+{
+public:
+	FIntersectionNear(const TSharedPtr<FData> DataIn, const TSharedPtr<FContourList> ContoursIn, FContour* const ContourIn);
+
+	void BevelTillThis() override;
+};
+
+/** Intersection of point's normal with edge. */
+class FIntersectionFar final : public FIntersection
+{
+public:
+	FIntersectionFar(const TSharedPtr<FData> DataIn, const TSharedPtr<FContourList> ContoursIn, FContour* const ContourIn);
+
+	void BevelTillThis() override;
+
+private:
+	/** Edge with which point's normal intersects. */
+	FPart* SplitEdge;
+
+
+	/**
+	 * Update list of expands for contour (after beveling till IntersectionFar contour is split to 2 contours).
+	 * @param UpdatedContour - Updated part.
+	 * @param OtherContour - Not updated part.
+	 * @param Curr - Point that splits contours.
+	 * @param SplitEdgePart - Part of split edge that belongs to UpdatedContour.
+	 * @param PrevDelta - Distance that Curr->Prev passed during FIntersectionFar::BevelTillThis.
+	 * @param NextDelta - Distance that Curr->Next passed during FIntersectionFar::BevelTillThis.
+	 */
+	void UpdateExpandsFar(FContour* const UpdatedContour, const FContour& OtherContour, FPart* const Curr, const FPart* const SplitEdgePart, const float PrevDelta, const float NextDelta);
+};
diff --git a/Engine/Plugins/Experimental/Text3D/Source/Text3D/Private/Mesh.cpp b/Engine/Plugins/Experimental/Text3D/Source/Text3D/Private/Mesh.cpp
new file mode 100644
--- /dev/null
+++ b/Engine/Plugins/Experimental/Text3D/Source/Text3D/Private/Mesh.cpp
@@ -0,0 +1,8 @@
+// Copyright 1998-2019 Epic Games, Inc. All Rights Reserved.
+
+#include "Mesh.h"
+
+bool FText3DMesh::IsEmpty() const
+{
+	return Vertices.Num() == 0 || Indices.Num() == 0;
+}
diff --git a/Engine/Plugins/Experimental/Text3D/Source/Text3D/Private/MeshCreator.cpp b/Engine/Plugins/Experimental/Text3D/Source/Text3D/Private/MeshCreator.cpp
new file mode 100644
--- /dev/null
+++ b/Engine/Plugins/Experimental/Text3D/Source/Text3D/Private/MeshCreator.cpp
@@ -0,0 +1,687 @@
+﻿// Copyright 1998-2019 Epic Games, Inc. All Rights Reserved.
+
+
+#include "MeshCreator.h"
+#include "Contour.h"
+#include "Part.h"
+#include "Intersection.h"
+#include "Data.h"
+#include "ContourList.h"
+
+
+THIRD_PARTY_INCLUDES_START
+#if PLATFORM_WINDOWS
+#include "Windows/WindowsHWrapper.h"
+#include "Windows/AllowWindowsPlatformTypes.h"
+#endif //PLATFORM_WINDOWS
+
+extern "C" {
+#include <tessellate.h>
+}
+
+#if PLATFORM_WINDOWS
+#include "Windows/HideWindowsPlatformTypes.h"
+#endif //PLATFORM_WINDOWS
+THIRD_PARTY_INCLUDES_END
+
+
+FMeshCreator::FMeshCreator(TSharedRef<TText3DMeshList> MeshesIn, const TSharedPtr<FData> DataIn)
+	: Meshes(MeshesIn)
+	, Data(DataIn)
+{
+
+}
+
+void FMeshCreator::CreateMeshes(const TSharedPtr<FContourList> ContoursIn, const float Extrude, const float Bevel, const EText3DBevelType Type, const int32 HalfCircleSegments)
+{
+	Contours = ContoursIn;
+
+	CreateFrontMesh();
+	CreateBevelMesh(Bevel, Type, HalfCircleSegments);
+#if !TEXT3D_WITH_INTERSECTION
+	CreateExtrudeMesh(Extrude, Bevel);
+#else
+	if (Extrude > Bevel)
+	{
+		Data->SetCurrentMesh(EText3DMeshType::Extrude);
+		const FVector2D Normal(1, 0);
+		BevelLinear.BevelContours(Extrude - Bevel, 0, Normal, Normal, false, MarkedVertex);
+	}
+
+	Data->SetCurrentMesh(EText3DMeshType::Back);
+	BevelLinear.CreateBackCap();
+#endif
+}
+
+void FMeshCreator::SetFrontAndBevelTextureCoordinates(const float Bevel)
+{
+	for (int32 Type = 0; Type < static_cast<int32>(EText3DMeshType::TypeCount); Type++)
+	{
+		Data->SetCurrentMesh(static_cast<EText3DMeshType>(Type));
+	}
+
+	struct FBoundingBox
+	{
+		FBox2D Box;
+		FVector2D Size;
+	};
+
+	TArray<FBoundingBox> BoundingBoxes;
+	FVector2D MaxSize(0, 0);
+	{
+		EText3DMeshType MeshType = FMath::IsNearlyZero(Bevel) ? EText3DMeshType::Front : EText3DMeshType::Bevel;
+		const FText3DMesh& Mesh = Meshes.Get()[static_cast<int32>(MeshType)];
+		const TArray<int32>& GlyphStartVertices = Mesh.GlyphStartVertices;
+		BoundingBoxes.AddUninitialized(GlyphStartVertices.Num() - 1);
+
+		for (int32 GlyphIndex = 0; GlyphIndex < BoundingBoxes.Num(); GlyphIndex++)
+		{
+			FBoundingBox& BoundingBox = BoundingBoxes[GlyphIndex];
+			FBox2D& Box = BoundingBox.Box;
+
+			const int32 FirstIndex = GlyphStartVertices[GlyphIndex];
+			const int32 LastIndex = GlyphStartVertices[GlyphIndex + 1];
+
+			if (FirstIndex < LastIndex)
+			{
+				const FVector& Position = Mesh.Vertices[FirstIndex].Position;
+				const FVector2D PositionFlat = {Position.Y, Position.Z};
+
+				Box.Min = PositionFlat;
+				Box.Max = PositionFlat;
+			}
+
+			for (int32 VertexIndex = FirstIndex + 1; VertexIndex < LastIndex; VertexIndex++)
+			{
+				const FDynamicMeshVertex& Vertex = Mesh.Vertices[VertexIndex];
+				const FVector& Position = Vertex.Position;
+
+				Box.Min.X = FMath::Min(Box.Min.X, Position.Y);
+				Box.Min.Y = FMath::Min(Box.Min.Y, Position.Z);
+				Box.Max.X = FMath::Max(Box.Max.X, Position.Y);
+				Box.Max.Y = FMath::Max(Box.Max.Y, Position.Z);
+			}
+
+			FVector2D& Size = BoundingBox.Size;
+			Size = Box.GetSize();
+
+			MaxSize.X = FMath::Max(MaxSize.X, Size.X);
+			MaxSize.Y = FMath::Max(MaxSize.Y, Size.Y);
+		}
+	}
+
+	for (int32 GlyphIndex = 0; GlyphIndex < BoundingBoxes.Num(); GlyphIndex++)
+	{
+		const int32 NextGlyphIndex = GlyphIndex + 1;
+		FBoundingBox& BoundingBox = BoundingBoxes[GlyphIndex];
+
+		TSharedPtr<TText3DMeshList> MeshesLocal = Meshes;
+		auto SetTextureCoordinates = [MeshesLocal, NextGlyphIndex, GlyphIndex, &BoundingBox, MaxSize](const EText3DMeshType Mesh)
+		{
+			FText3DMesh& CurrentMesh = (*MeshesLocal.Get())[static_cast<int32>(Mesh)];
+			const TArray<int32>& GlyphStartVertices = CurrentMesh.GlyphStartVertices;
+
+			if (NextGlyphIndex >= GlyphStartVertices.Num())
+			{
+				return;
+			}
+
+			for (int32 VertexIndex = GlyphStartVertices[GlyphIndex]; VertexIndex < GlyphStartVertices[NextGlyphIndex]; VertexIndex++)
+			{
+				FDynamicMeshVertex& Vertex = CurrentMesh.Vertices[VertexIndex];
+				const FVector2D TextureCoordinate = (FVector2D(Vertex.Position.Y, Vertex.Position.Z) - BoundingBox.Box.Min) / MaxSize;
+
+				for (int32 Index = 0; Index < MAX_STATIC_TEXCOORDS; Index++)
+				{
+					Vertex.TextureCoordinate[Index] = {TextureCoordinate.X, 1 - TextureCoordinate.Y};
+				}
+			}
+		};
+
+		SetTextureCoordinates(EText3DMeshType::Front);
+		SetTextureCoordinates(EText3DMeshType::Bevel);
+	}
+}
+
+void FMeshCreator::MirrorMeshes(const float Extrude, const float ScaleX)
+{
+	// When TEXT3D_WITH_INTERSECTION=1 the following functions will not do anything
+	MirrorMesh(EText3DMeshType::Bevel, EText3DMeshType::Bevel, Extrude, ScaleX);
+	MirrorMesh(EText3DMeshType::Front, EText3DMeshType::Back, Extrude, ScaleX);
+}
+
+// Using GLUtesselator to make triangulation of front mesh
+void FMeshCreator::CreateFrontMesh()
+{
+	int32 VerticesCount = 0;
+
+	for (const FContour& Contour : *Contours)
+	{
+		VerticesCount += Contour.Num();
+	}
+
+	const TUniquePtr<double[]> Vertices = MakeUnique<double[]>(SIZE_T(VerticesCount) * 2);
+	const TUniquePtr<const double* []> ContoursIn = MakeUnique<const double* []>(SIZE_T(Contours->Num()) + 1);
+	ContoursIn[0] = Vertices.Get();
+
+	int32 Offset = 0;
+	int32 ContourIndex = 0;
+	for (const FContour& Contour : *Contours)
+	{
+		int32 PointIndex = 0;
+		const FPart* const First = Contour[0];
+
+		for (const FPart* Point = First; ; Point = Point->Next, PointIndex++)
+		{
+			double* const Vertex = Vertices.Get() + (Offset + PointIndex) * 2;
+			const FVector2D Position = Point->Position;
+
+			Vertex[0] = double(Position.X);
+			Vertex[1] = double(Position.Y);
+
+			if (Point == First->Prev)
+			{
+				break;
+			}
+		}
+
+		Offset += Contour.Num();
+		ContoursIn[SIZE_T(ContourIndex++ + 1)] = Vertices.Get() + Offset * 2;
+	}
+
+	double* VerticesOut = nullptr;
+	int* IndicesOut = nullptr;
+	int VerticesCountOut = 0;
+	int TrianglesCountOut = 0;
+
+	// @third party code - BEGIN GLUtesselator tesselate
+	// Used to tesselate mesh framed with given contours
+	tessellate(&VerticesOut, &VerticesCountOut, &IndicesOut, &TrianglesCountOut, static_cast<const double**>(ContoursIn.Get()), static_cast<const double**>(ContoursIn.Get() + ContourIndex + 1));
+	//@third party code - END GLUtesselator tesselate
+
+
+	auto DeleteData = [&VerticesOut, &IndicesOut]()
+	{
+		if (IndicesOut)
+			free(IndicesOut);
+
+		if (VerticesOut)
+			free(VerticesOut);
+
+		IndicesOut = nullptr;
+		VerticesOut = nullptr;
+	};
+
+
+	if (VerticesCountOut != VerticesCount)
+	{
+		DeleteData();
+		Contours->Empty();
+		return;
+	}
+
+
+	Data->SetCurrentMesh(EText3DMeshType::Front);
+
+
+	Data->ResetDoneExtrude();
+	Data->SetMinBevelTarget();
+	const int32 FirstAdded = Data->AddVertices(VerticesCount);
+
+	for (const FContour& Contour : *Contours)
+	{
+		const FPart* const First = Contour[0];
+
+		for (const FPart* Point = First; ; Point = Point->Next)
+		{
+			Data->AddVertex(Point, {1, 0}, {0, 0, -1});
+
+			if (Point == First->Prev)
+			{
+				break;
+			}
+		}
+	}
+
+
+	Data->AddTriangles(TrianglesCountOut);
+
+	for (int Index = 0; Index < TrianglesCountOut; Index++)
+	{
+		const int* const t = IndicesOut + Index * 3;
+		Data->AddTriangle(FirstAdded + t[0], FirstAdded + t[1], FirstAdded + t[2]);
+	}
+
+
+	DeleteData();
+}
+
+void FMeshCreator::CreateBevelMesh(const float Bevel, const EText3DBevelType Type, const int32 HalfCircleSegments)
+{
+	if (FMath::IsNearlyZero(Bevel))
+	{
+		return;
+	}
+
+	Data->SetCurrentMesh(EText3DMeshType::Bevel);
+
+	switch (Type)
+	{
+	case EText3DBevelType::Linear:
+	{
+		const FVector2D Normal = FVector2D(1, -1).GetSafeNormal();
+		BevelLinear(Bevel, Bevel, Normal, Normal, false);
+		break;
+	}
+	case EText3DBevelType::HalfCircle:
+	{
+		const float Step = HALF_PI / HalfCircleSegments;
+
+		float CosCurr = 1.0f;
+		float SinCurr = 0.0f;
+
+		float CosNext = 0.0f;
+		float SinNext = 0.0f;
+
+		float ExtrudeLocalNext = 0.0f;
+		float ExpandLocalNext = 0.0f;
+
+		bool bSmoothNext = false;
+
+		FVector2D NormalNext;
+		FVector2D NormalEnd;
+
+		auto MakeStep = [&SinNext, &CosNext, Step, &ExtrudeLocalNext, &ExpandLocalNext, Bevel, &CosCurr, &SinCurr, &NormalNext](int32 Index)
+		{
+			FMath::SinCos(&SinNext, &CosNext, Index * Step);
+
+			ExtrudeLocalNext = Bevel * (CosCurr - CosNext);
+			ExpandLocalNext = Bevel * (SinNext - SinCurr);
+
+			NormalNext = FVector2D(ExtrudeLocalNext, -ExpandLocalNext).GetSafeNormal();
+		};
+
+		MakeStep(1);
+		for (int32 Index = 0; Index < HalfCircleSegments; Index++)
+		{
+			CosCurr = CosNext;
+			SinCurr = SinNext;
+
+			const float ExtrudeLocal = ExtrudeLocalNext;
+			const float ExpandLocal = ExpandLocalNext;
+
+			const FVector2D Normal = NormalNext;
+			FVector2D NormalStart;
+
+			const bool bFirst = (Index == 0);
+			const bool bLast = (Index == HalfCircleSegments - 1);
+
+			const bool bSmooth = bSmoothNext;
+
+			if (!bLast)
+			{
+				MakeStep(Index + 2);
+				bSmoothNext = FVector2D::DotProduct(Normal, NormalNext) >= -FPart::CosMaxAngleSides;
+			}
+
+			NormalStart = bFirst ? Normal : (bSmooth ? NormalEnd : Normal);
+			NormalEnd = bLast ? Normal : (bSmoothNext ? (Normal + NormalNext).GetSafeNormal() : Normal);
+
+			BevelLinear(ExtrudeLocal, ExpandLocal, NormalStart, NormalEnd, bSmooth);
+		}
+
+		break;
+	}
+
+	default:
+		break;
+	}
+}
+
+void FMeshCreator::CreateExtrudeMesh(float Extrude, const float Bevel)
+{
+	if (Bevel >= Extrude / 2)
+	{
+		return;
+	}
+
+	Data->SetCurrentMesh(EText3DMeshType::Extrude);
+	Extrude -= Bevel * 2.0f;
+	Data->SetExtrude(Extrude);
+	Data->SetExpand(0);
+
+	const FVector2D Normal(1, 0);
+	Data->SetNormals(Normal, Normal);
+
+	for (FContour& Contour : *Contours)
+	{
+		for (FPart* const Part : Contour)
+		{
+			Part->ResetDoneExpand();
+		}
+	}
+
+
+	TArray<float> TextureCoordinateVs;
+
+	auto EdgeLength = [](const FPart* const Edge)
+	{
+		return (Edge->Next->Position - Edge->Position).Size();
+	};
+
+	for (FContour& Contour : *Contours)
+	{
+		// Compute TexCoord.V-s for each point
+		TextureCoordinateVs.Reset(Contour.Num() - 1);
+		const FPart* const First = Contour[0];
+		TextureCoordinateVs.Add(EdgeLength(First));
+
+		int32 Index = 1;
+		for (const FPart* Edge = First->Next; Edge != First->Prev; Edge = Edge->Next)
+		{
+			TextureCoordinateVs.Add(TextureCoordinateVs[Index - 1] + EdgeLength(Edge));
+			Index++;
+		}
+
+
+		const float ContourLength = TextureCoordinateVs.Last() + EdgeLength(Contour.Last());
+
+		if (FMath::IsNearlyZero(ContourLength))
+		{
+			continue;
+		}
+
+
+		for (float& PointY : TextureCoordinateVs)
+		{
+			PointY /= ContourLength;
+		}
+
+		// Duplicate contour
+		Data->SetMinBevelTarget();
+
+		// First point in contour is processed separately
+		{
+			FPart* const Point = Contour[0];
+			// It's set to sharp because we need 2 vertices with TexCoord.Y values 0 and 1 (for smooth points only one vertex is added)
+			Point->bSmooth = false;
+			EmptyPaths(Point);
+			ExpandPointWithoutAddingVertices(Point);
+
+			const FVector2D TexCoordPrev(0, 0);
+			const FVector2D TexCoordCurr(0, 1);
+
+			if (Point->bSmooth)
+			{
+				AddVertexSmooth(Point, TexCoordPrev);
+				AddVertexSmooth(Point, TexCoordCurr);
+			}
+			else
+			{
+				AddVertexSharp(Point, Point->Prev, TexCoordPrev);
+				AddVertexSharp(Point, Point, TexCoordCurr);
+			}
+		}
+
+		Index = 1;
+		for (FPart* Point = First->Next; Point != First; Point = Point->Next)
+		{
+			EmptyPaths(Point);
+			ExpandPoint(Point, {0, 1 - TextureCoordinateVs[Index++ - 1]});
+		}
+
+
+		// Add extruded vertices
+		Data->SetMaxBevelTarget();
+
+		// Similarly to duplicating vertices, first point is processed separately
+		{
+			FPart* const Point = Contour[0];
+			ExpandPointWithoutAddingVertices(Point);
+
+			const FVector2D TexCoordPrev(1, 0);
+			const FVector2D TexCoordCurr(1, 1);
+
+			if (Point->bSmooth)
+			{
+				AddVertexSmooth(Point, TexCoordPrev);
+				AddVertexSmooth(Point, TexCoordCurr);
+			}
+			else
+			{
+				AddVertexSharp(Point, Point->Prev, TexCoordPrev);
+				AddVertexSharp(Point, Point, TexCoordCurr);
+			}
+		}
+
+		Index = 1;
+		for (FPart* Point = First->Next; Point != First; Point = Point->Next)
+		{
+			ExpandPoint(Point, {1, 1 - TextureCoordinateVs[Index++ - 1]});
+		}
+
+
+		for (FPart* const Edge : Contour)
+		{
+			Data->FillEdge(Edge, false);
+		}
+	}
+}
+
+void FMeshCreator::MirrorMesh(const EText3DMeshType TypeIn, const EText3DMeshType TypeOut, const float Extrude, const float ScaleX)
+{
+#if !TEXT3D_WITH_INTERSECTION
+	const FText3DMesh& MeshIn = Meshes.Get()[static_cast<int32>(TypeIn)];
+	FText3DMesh& MeshOut = Meshes.Get()[static_cast<int32>(TypeOut)];
+
+
+	const TArray<FDynamicMeshVertex>& VerticesIn = MeshIn.Vertices;
+	TArray<FDynamicMeshVertex>& VerticesOut = MeshOut.Vertices;
+
+	const int32 VerticesInNum = VerticesIn.Num();
+	const int32 VerticesOutNum = VerticesOut.Num();
+
+
+	const TArray<int32>& IndicesIn = MeshIn.Indices;
+	TArray<int32>& IndicesOut = MeshOut.Indices;
+
+	const int32 IndicesInNum = IndicesIn.Num();
+	const int32 IndicesOutNum = IndicesOut.Num();
+
+
+
+	VerticesOut.AddUninitialized(VerticesInNum);
+
+	for (int32 Index = 0; Index < VerticesInNum; Index++)
+	{
+		const FDynamicMeshVertex& Vertex = VerticesIn[Index];
+
+		const FVector& Position = Vertex.Position;
+		const FVector TangentX = Vertex.TangentX.ToFVector();
+		const FVector TangentZ = Vertex.TangentZ.ToFVector();
+
+		VerticesOut[VerticesOutNum + Index] = FDynamicMeshVertex({Extrude * ScaleX - Position.X, Position.Y, Position.Z}, {-TangentX.X, TangentX.Y, TangentX.Z}, {-TangentZ.X, TangentZ.Y, TangentZ.Z}, Vertex.TextureCoordinate[0], Vertex.Color);
+	}
+
+
+	IndicesOut.AddUninitialized(IndicesInNum);
+
+	for (int32 Index = 0; Index < IndicesInNum / 3; Index++)
+	{
+		const int32 OldTriangle = Index * 3;
+		const int32 NewTriangle = IndicesOutNum + OldTriangle;
+
+		IndicesOut[NewTriangle + 0] = VerticesOutNum + IndicesIn[OldTriangle + 0];
+		IndicesOut[NewTriangle + 1] = VerticesOutNum + IndicesIn[OldTriangle + 2];
+		IndicesOut[NewTriangle + 2] = VerticesOutNum + IndicesIn[OldTriangle + 1];
+	}
+#endif
+}
+
+void FMeshCreator::BevelLinear(const float Extrude, const float Expand, FVector2D NormalStart, FVector2D NormalEnd, const bool bSmooth)
+{
+	Reset(Extrude, Expand, NormalStart, NormalEnd);
+
+	if (!bSmooth)
+	{
+		DuplicateContourVertices();
+	}
+
+	if (Expand > 0)
+	{
+		BevelPartsWithIntersectingNormals();
+	}
+
+	BevelPartsWithoutIntersectingNormals();
+
+	Data->IncreaseDoneExtrude();
+}
+
+void FMeshCreator::DuplicateContourVertices()
+{
+	Data->SetMinBevelTarget();
+
+	for (FContour& Contour : *Contours)
+	{
+		for (FPart* const Point : Contour)
+		{
+			EmptyPaths(Point);
+			// Duplicate points of contour (expansion with value 0)
+			ExpandPoint(Point);
+		}
+	}
+}
+
+void FMeshCreator::Reset(const float Extrude, const float Expand, FVector2D NormalStart, FVector2D NormalEnd)
+{
+	Data->SetExtrude(Extrude);
+	Data->SetExpand(Expand);
+
+	Data->SetNormals(NormalStart, NormalEnd);
+	Contours->Reset();
+}
+
+template<class T>
+FIntersection& Upcast(T& Intersection)
+{
+	return static_cast<FIntersection&>(Intersection);
+}
+
+void FMeshCreator::BevelPartsWithIntersectingNormals()
+{
+#if TEXT3D_WITH_INTERSECTION
+	for (Iteration = 0; /*Iteration < Iterations*/true; Iteration++)
+	{
+		// Copy list of contours (but not contours themselves) and iterate this list
+		// because contours can be added or removed while beveling till intersections
+		TArray<FContour*> ContoursCopy;
+
+		for (FContour& Contour : *Contours)
+		{
+			ContoursCopy.Add(&Contour);
+		}
+
+		bool bIntersectionsExisted = false;
+
+		for (FContour* const Contour : ContoursCopy)
+		{
+			if (!Contour->HasIntersections())
+			{
+				continue;
+			}
+
+			FIntersectionNear Near(Data, Contours, Contour);
+			FIntersectionFar Far(Data, Contours, Contour);
+
+			FIntersection& Closest = Near.GetValue() <= Far.GetValue() ? Upcast(Near) : Upcast(Far);
+			const float Value = Closest.GetValue();
+
+			// If intersection will happen further from front cap then is needed to bevel to, skip
+			if (Value > Data->GetExpand())
+			{
+				Contour->DisableIntersections();
+				continue;
+			}
+
+			Data->SetExpandTarget(Value);
+			Closest.BevelTillThis();
+			bIntersectionsExisted = true;
+		}
+
+		if (!bIntersectionsExisted)
+		{
+			break;
+		}
+	}
+#endif
+}
+
+void FMeshCreator::BevelPartsWithoutIntersectingNormals()
+{
+	Data->SetMaxBevelTarget();
+	const float MaxExpand = Data->GetExpand();
+
+	for (FContour& Contour : *Contours)
+	{
+		for (FPart* const Point : Contour)
+		{
+			if (!FMath::IsNearlyEqual(Point->DoneExpand, MaxExpand) || FMath::IsNearlyZero(MaxExpand))
+			{
+				ExpandPoint(Point);
+			}
+
+			const float Delta = MaxExpand - Point->DoneExpand;
+
+			Point->AvailableExpandNear -= Delta;
+			Point->DecreaseExpandsFar(Delta);
+		}
+
+		for (FPart* const Edge : Contour)
+		{
+			Data->FillEdge(Edge, false);
+		}
+	}
+}
+
+void FMeshCreator::EmptyPaths(FPart* const Point) const
+{
+	Point->PathPrev.Empty();
+	Point->PathNext.Empty();
+}
+
+void FMeshCreator::ExpandPoint(FPart* const Point, const FVector2D TextureCoordinates)
+{
+	ExpandPointWithoutAddingVertices(Point);
+
+	if (Point->bSmooth)
+	{
+		AddVertexSmooth(Point, TextureCoordinates);
+	}
+	else
+	{
+		AddVertexSharp(Point, Point->Prev, TextureCoordinates);
+		AddVertexSharp(Point, Point, TextureCoordinates);
+	}
+}
+
+void FMeshCreator::ExpandPointWithoutAddingVertices(FPart* const Point) const
+{
+	Point->Position = Data->Expanded(Point);
+	const int32 FirstAdded = Data->AddVertices(Point->bSmooth ? 1 : 2);
+
+	Point->PathPrev.Add(FirstAdded);
+	Point->PathNext.Add(Point->bSmooth ? FirstAdded : FirstAdded + 1);
+}
+
+void FMeshCreator::AddVertexSmooth(const FPart* const Point, const FVector2D TextureCoordinates)
+{
+	const FPart* const Curr = Point;
+	const FPart* const Prev = Point->Prev;
+
+	Data->AddVertex(Point, (Prev->TangentX + Curr->TangentX).GetSafeNormal(), (Data->ComputeTangentZ(Prev, Point->DoneExpand) + Data->ComputeTangentZ(Curr, Point->DoneExpand)).GetSafeNormal(), TextureCoordinates);
+}
+
+void FMeshCreator::AddVertexSharp(const FPart* const Point, const FPart* const Edge, const FVector2D TextureCoordinates)
+{
+	Data->AddVertex(Point, Edge->TangentX, Data->ComputeTangentZ(Edge, Point->DoneExpand).GetSafeNormal(), TextureCoordinates);
+}
diff --git a/Engine/Plugins/Experimental/Text3D/Source/Text3D/Private/MeshCreator.h b/Engine/Plugins/Experimental/Text3D/Source/Text3D/Private/MeshCreator.h
new file mode 100644
--- /dev/null
+++ b/Engine/Plugins/Experimental/Text3D/Source/Text3D/Private/MeshCreator.h
@@ -0,0 +1,122 @@
+// Copyright 1998-2019 Epic Games, Inc. All Rights Reserved.
+
+#pragma once
+
+#include "BevelType.h"
+#include "Mesh.h"
+
+struct FPart;
+class FData;
+class FContourList;
+
+/** Makes actual bevel. */
+class FMeshCreator final
+{
+public:
+	FMeshCreator(TSharedRef<TText3DMeshList> MeshesIn, const TSharedPtr<FData> DataIn);
+
+	/**
+	 * Create meshes.
+	 * @param ContoursIn - List of contours.
+	 * @param Extrude - Orthogonal (to front cap) offset value.
+	 * @param Bevel - Bevel value (bevel happens before extrude).
+	 * @param Type - Defines shape of beveled part.
+	 * @param HalfCircleSegments - Segments count for Type == EText3DBevelType::HalfCircle.
+	 */
+	void CreateMeshes(const TSharedPtr<FContourList> ContoursIn, const float Extrude, const float Bevel, const EText3DBevelType Type, const int32 HalfCircleSegments);
+	void SetFrontAndBevelTextureCoordinates(const float Bevel);
+	void MirrorMeshes(const float Extrude, const float ScaleX);
+
+private:
+	TSharedRef<TText3DMeshList> Meshes;
+	const TSharedPtr<FData> Data;
+	TSharedPtr<FContourList> Contours;
+
+
+	/**
+	 * Create 'Front' part of glyph.
+	 */
+	void CreateFrontMesh();
+	/**
+	 *
+	 */
+	/**
+	 * Create 'Bevel' part of glyph (actually half of it, will be mirrored later).
+	 * @param Bevel - Bevel value (bevel happens before extrude).
+	 * @param Type - Defines shape of beveled part.
+	 * @param HalfCircleSegments - Segments count for Type == EText3DBevelType::HalfCircle.
+	 */
+	void CreateBevelMesh(const float Bevel, const EText3DBevelType Type, const int32 HalfCircleSegments);
+	/**
+	 * Create 'Extrude' part of glyph.
+	 * @param Extrude - Extrude value.
+	 * @param Bevel - Bevel value (bevel happens before extrude).
+	 */
+	void CreateExtrudeMesh(float Extrude, const float Bevel);
+
+	void MirrorMesh(const EText3DMeshType TypeIn, const EText3DMeshType TypeOut, const float Extrude, const float ScaleX);
+
+
+	/**
+	 * Bevel one segment.
+	 * @param Extrude - Documented in ContourList.h.
+	 * @param Expand - Documented in ContourList.h.
+	 * @param NormalStart - Normal at start of segment (minimum DoneExpand).
+	 * @param NormalEnd - Normal at end of segment (maximum DoneExpand).
+	 * @param bSmooth - Is angle between start of this segment and end of previous segment smooth?
+	 */
+	void BevelLinear(const float Extrude, const float Expand, FVector2D NormalStart, FVector2D NormalEnd, const bool bSmooth);
+
+	/**
+	 * Duplicate contour vertices (used to make sharp angle between bevel steps)
+	 */
+	void DuplicateContourVertices();
+
+	/**
+	 * Prepare for beveling (is executed before each step).
+	 */
+	void Reset(const float Extrude, const float Expand, FVector2D NormalStart, FVector2D NormalEnd);
+
+	/**
+	 * Make bevel only for non-trivial places.
+	 */
+	void BevelPartsWithIntersectingNormals();
+
+	/**
+	 * Continue with trivial bevel till FData::Expand.
+	 */
+	void BevelPartsWithoutIntersectingNormals();
+
+	/**
+	 * Clear PathPrev and PathNext.
+	 * @param Point - Point which paths should be cleared.
+	 */
+	void EmptyPaths(FPart* const Point) const;
+
+	/**
+	 * Same as previous function but does not cover the intersection-case.
+	 * @param Point - Point that should be expanded.
+	 * @param TextureCoordinates - Texcture coordinates of added vertices.
+	 */
+	void ExpandPoint(FPart* const Point, const FVector2D TextureCoordinates = FVector2D(0, 0));
+	/**
+	 * Common code for expanding, vertices are added uninitialized.
+	 * @param Point - Expanded point.
+	 */
+	void ExpandPointWithoutAddingVertices(FPart* const Point) const;
+
+	/**
+	 * Add vertex for smooth point.
+	 * @param Point - Expanded point.
+	 * @param TextureCoordinates - Texture coordinates of added vertex.
+	 */
+	void AddVertexSmooth(const FPart* const Point, const FVector2D TextureCoordinates);
+	/**
+	 * Add vertex for sharp point.
+	 * @param Point - Expanded point.
+	 * @param Edge - Edge from which TangentX and TangentZ will be assigned.
+	 * @param TextureCoordinates - Texture coordinates of added vertex.
+	 */
+	void AddVertexSharp(const FPart* const Point, const FPart* const Edge, const FVector2D TextureCoordinates);
+
+};
diff --git a/Engine/Plugins/Experimental/Text3D/Source/Text3D/Private/Part.cpp b/Engine/Plugins/Experimental/Text3D/Source/Text3D/Private/Part.cpp
new file mode 100644
--- /dev/null
+++ b/Engine/Plugins/Experimental/Text3D/Source/Text3D/Private/Part.cpp
@@ -0,0 +1,110 @@
+// Copyright 1998-2019 Epic Games, Inc. All Rights Reserved.
+
+#include "Part.h"
+#include "Util.h"
+
+#include "HAL/PlatformMath.h"
+
+constexpr float FPart::CosMaxAngleSideTangent;
+constexpr float FPart::CosMaxAngleSides;
+
+FPart::FPart()
+{
+	Prev = nullptr;
+	Next = nullptr;
+	DoneExpand = 0.0f;
+	bSmooth = false;
+	AvailableExpandNear = 0.0f;
+}
+
+void FPart::ResetDoneExpand()
+{
+	DoneExpand = 0;
+}
+
+void FPart::ComputeTangentX()
+{
+	check(Next);
+	TangentX = (Next->Position - Position).GetSafeNormal();
+}
+
+bool FPart::ComputeNormal()
+{
+	check(Prev);
+
+	const FVector2D A = -Prev->TangentX;
+	const FVector2D C = TangentX;
+
+	Normal = A + C;
+
+	const float NormalLength2 = Normal.SizeSquared();
+
+	// Scale is needed to make ((p_(i+1) + k * n_(i+1)) - (p_i + k * n_i)) parallel to (p_(i+1) - p_i). Also (k) is distance between original edge and this edge after expansion with value (k).
+	const float OneMinusADotC = 1.0f - FVector2D::DotProduct(A, C);
+
+	if (FMath::IsNearlyZero(OneMinusADotC))
+	{
+		return false;
+	}
+
+	const float Scale = -FPlatformMath::Sqrt(2.0f / OneMinusADotC);
+
+
+	// If previous and next edge are nearly on one line
+	if (FMath::IsNearlyZero(NormalLength2, 0.0001f))
+	{
+		Normal = FVector2D(A.Y, -A.X) * Scale;
+	}
+	else
+	{
+		// Sign of cross product is needed to be sure that Normal is directed outside.
+		Normal *= -Scale * FPlatformMath::Sign(FVector2D::CrossProduct(A, C)) / FPlatformMath::Sqrt(NormalLength2);
+	}
+
+	return true;
+}
+
+void FPart::ComputeSmooth()
+{
+	check(Prev);
+	bSmooth = FVector2D::DotProduct(-Prev->TangentX, TangentX) <= CosMaxAngleSides;
+}
+
+bool FPart::ComputeNormalAndSmooth()
+{
+	if (!ComputeNormal())
+	{
+		return false;
+	}
+
+	ComputeSmooth();
+	return true;
+}
+
+void FPart::ResetInitialPosition()
+{
+	InitialPosition = Position;
+}
+
+void FPart::ComputeInitialPosition()
+{
+	InitialPosition = Position - DoneExpand * Normal;
+}
+
+void FPart::DecreaseExpandsFar(const float Delta)
+{
+	for (auto i = AvailableExpandsFar.CreateIterator(); i; ++i)
+	{
+		i->Value -= Delta;
+
+		if (i->Value < 0)
+		{
+			i.RemoveCurrent();
+		}
+	}
+}
+
+FVector2D FPart::Expanded(const float Value) const
+{
+	return Position + Normal * Value;
+}
diff --git a/Engine/Plugins/Experimental/Text3D/Source/Text3D/Private/Part.h b/Engine/Plugins/Experimental/Text3D/Source/Text3D/Private/Part.h
new file mode 100644
--- /dev/null
+++ b/Engine/Plugins/Experimental/Text3D/Source/Text3D/Private/Part.h
@@ -0,0 +1,77 @@
+// Copyright 1998-2019 Epic Games, Inc. All Rights Reserved.
+
+#pragma once
+
+#include "CoreMinimal.h"
+#include "Util.h"
+
+/** Part represents point and it's next edge (counterclockwise one). */
+struct FPart final
+{
+	static constexpr float CosMaxAngleSideTangent = 0.995f;
+	static constexpr float CosMaxAngleSides = 1.f - 2.f * CosMaxAngleSideTangent * CosMaxAngleSideTangent;
+
+	FPart();
+
+
+	/** Previous part. */
+	FPart* Prev;
+
+	/** Next part. */
+	FPart* Next;
+
+
+	/** Position, is equal to position of last vertex in paths (in coordinate system of glyph). */
+	FVector2D Position;
+
+	/** Offset in surface of front cap that this point already made. */
+	float DoneExpand;
+
+
+	FVector2D TangentX;
+
+
+	/** Point normal, a bisector of angle. */
+	FVector2D Normal;
+
+	/** If true, previous and next edges are in one smoothing group. */
+	bool bSmooth;
+
+
+	FVector2D InitialPosition;
+
+
+	// Paths along which triangulation is made. Values that are stored are indices of vertices. If bSmooth == true, both paths store one index (for same DoneExpand value). If not, indices are different.
+	/** Path used for triangulation of previous edge. */
+	TArray<int32> PathPrev;
+
+	/** Path used for triangulation of next edge. */
+	TArray<int32> PathNext;
+
+
+	/** Offset needed for an IntersectionNear to happen. */
+	float AvailableExpandNear;
+
+	/** List of pairs (edge, offset) for IntersectionFar. */
+	FAvailableExpandsFar AvailableExpandsFar;
+
+
+	void ResetDoneExpand();
+	void ComputeTangentX();
+	bool ComputeNormal();
+	void ComputeSmooth();
+
+	bool ComputeNormalAndSmooth();
+
+	void ResetInitialPosition();
+	void ComputeInitialPosition();
+
+	void DecreaseExpandsFar(const float Delta);
+
+	/**
+	 * Compute position to which point will be expanded.
+	 * @param Value - Offset in surface of front cap.
+	 * @return Computed position.
+	 */
+	FVector2D Expanded(const float Value) const;
+};
diff --git a/Engine/Plugins/Experimental/Text3D/Source/Text3D/Private/Text3DComponent.cpp b/Engine/Plugins/Experimental/Text3D/Source/Text3D/Private/Text3DComponent.cpp
--- a/Engine/Plugins/Experimental/Text3D/Source/Text3D/Private/Text3DComponent.cpp
+++ b/Engine/Plugins/Experimental/Text3D/Source/Text3D/Private/Text3DComponent.cpp
@@ -2,8 +2,9 @@
 
 
 #include "Text3DComponent.h"
+#include "Text3DPrivate.h"
 
-#include "Bevel.h"
+#include "ContourList.h"
 #include "Data.h"
 #include "Engine/Font.h"
 #include "Engine/Engine.h"
@@ -15,30 +16,18 @@
 #include "Text3DPrivate.h"
 #include "TextShaper.h"
 #include "UObject/ConstructorHelpers.h"
+#include "MeshCreator.h"
 
 #include "Misc/EngineVersionComparison.h"
 
 
-THIRD_PARTY_INCLUDES_START
-#if PLATFORM_WINDOWS
-#include "Windows/WindowsHWrapper.h"
-#include "Windows/AllowWindowsPlatformTypes.h"
-#endif //PLATFORM_WINDOWS
-#include "GL/glcorearb.h"
-#include "FTVectoriser.h"
-
-#if PLATFORM_WINDOWS
-#include "Windows/HideWindowsPlatformTypes.h"
-#endif //PLATFORM_WINDOWS
-THIRD_PARTY_INCLUDES_END
-
 #define LOCTEXT_NAMESPACE "Text3D"
 
 
 using TTextMeshDynamicData = TArray<TUniquePtr<FText3DDynamicData>, TFixedAllocator<static_cast<int32>(EText3DMeshType::TypeCount)>>;
 
 
-class FTextIndexBuffer : public FIndexBuffer
+class FTextIndexBuffer final : public FIndexBuffer
 {
 public:
 	virtual void InitRHI() override
@@ -65,7 +54,7 @@ public:
 		const TText3DMeshList & ComponentMeshes = Component->Meshes.Get();
 		for (int32 Index = 0; Index < static_cast<int32>(EText3DMeshType::TypeCount); Index++)
 		{
-			Meshes.Add(MakeUnique<FProxyMesh>(this, Index, ComponentMeshes[Index], Component->GetMaterial(Index)));
+			Meshes.Add(MakeUnique<FProxyMesh>(this, ComponentMeshes[Index], Component->GetMaterial(Index)));
 		}
 	}
 
@@ -152,23 +141,21 @@ private:
 		int32 VertexCount;
 		FStaticMeshVertexBuffers VertexBuffers;
 
-		int32 IndicesCount;
 		FTextIndexBuffer IndexBuffer;
 
 		UMaterialInterface* Material;
 		bool bInitialized;
 
 
-		FProxyMesh(const FText3DSceneProxy* const Proxy, const int32 MeshType, const FText3DMesh& ComponentMeshIn, UMaterialInterface * const InMaterial)
+		FProxyMesh(const FText3DSceneProxy* const Proxy, const FText3DMesh& ComponentMeshIn, UMaterialInterface * const InMaterial)
 			: ComponentMesh(ComponentMeshIn)
 
 			, VertexFactory(Proxy->GetScene().GetFeatureLevel(), "FText3DSceneProxyMesh")
 			, VertexCount(0)
 
-			, IndicesCount(0)
-
 			, Material(InMaterial)
 		{
+			IndexBuffer.NumIndices = 0;
 
 			if (Material == nullptr)
 			{
@@ -185,8 +172,7 @@ private:
 			VertexCount = ComponentMesh.Vertices.Num();
 			VertexBuffers.InitWithDummyData(&VertexFactory, uint32(VertexCount));
 
-			IndicesCount = ComponentMesh.Indices.Num();
-			IndexBuffer.NumIndices = IndicesCount;
+			IndexBuffer.NumIndices = ComponentMesh.Indices.Num();
 
 			BeginInitResource(&IndexBuffer);
 			bInitialized = true;
@@ -204,7 +190,7 @@ private:
 
 		bool IsEmpty() const
 		{
-			return VertexCount == 0 || IndicesCount == 0;
+			return VertexCount == 0 || IndexBuffer.NumIndices == 0;
 		}
 
 		void UpdateData()
@@ -212,7 +198,7 @@ private:
 			if (ComponentMesh.IsEmpty())
 			{
 				VertexCount = 0;
-				IndicesCount = 0;
+				IndexBuffer.NumIndices = 0;
 
 				bInitialized = false;
 				return;
@@ -221,13 +207,12 @@ private:
 			const int32 NewVertexCount = ComponentMesh.Vertices.Num();
 			const int32 NewIndicesCount = ComponentMesh.Indices.Num();
 
-			if (VertexCount != NewVertexCount || IndicesCount != NewIndicesCount)
+			if (VertexCount != NewVertexCount || IndexBuffer.NumIndices != NewIndicesCount)
 			{
 				VertexCount = NewVertexCount;
 				VertexBuffers.InitWithDummyData(&VertexFactory, uint32(VertexCount));
 
-				IndicesCount = NewIndicesCount;
-				IndexBuffer.NumIndices = IndicesCount;
+				IndexBuffer.NumIndices = NewIndicesCount;
 
 				if (bInitialized)
 				{
@@ -326,7 +311,7 @@ private:
 			BatchElement.PrimitiveUniformBufferResource = &DynamicPrimitiveUniformBuffer.UniformBuffer;
 
 			BatchElement.FirstIndex = 0;
-			BatchElement.NumPrimitives = uint32(IndicesCount) / 3;
+			BatchElement.NumPrimitives = uint32(IndexBuffer.NumIndices) / 3;
 			BatchElement.MinVertexIndex = 0;
 			BatchElement.MaxVertexIndex = uint32(VertexCount);
 			Mesh.ReverseCulling = Proxy->IsLocalToWorldDeterminantNegative();
@@ -383,15 +368,15 @@ UText3DComponent::UText3DComponent() :
 	MaxWidth = 500.0;
 	Bevel = 0.0f;
 	BevelType = EText3DBevelType::HalfCircle;
-	HalfCircleSegments = 4;
+	HalfCircleSegments = 8;
 
 	bHasMaxHeight = false;
 	MaxHeight = 500.0f;
 	bScaleProportionally = true;
 
 	bAutoActivate = true;
-	PendingBuild = false;
-	FreezeBuild = false;
+	bPendingBuild = false;
+	bFreezeBuild = false;
 }
 
 bool UText3DComponent::ShouldRecreateProxyOnUpdateTransform() const
@@ -463,7 +448,7 @@ void UText3DComponent::SetMaterial(int32 ElementIndex, UMaterialInterface* InMat
 	}
 	}
 
-	if (bChanged && !FreezeBuild)
+	if (bChanged && !bFreezeBuild)
 	{
 		MarkRenderStateDirty();
 	}
@@ -509,7 +494,7 @@ void UText3DComponent::SetText(const FText& Value)
 
 void UText3DComponent::SetKerning(const float Value)
 {
-	if (Kerning != Value)
+	if (!FMath::IsNearlyEqual(Kerning, Value))
 	{
 		Kerning = Value;
 		Rebuild();
@@ -518,7 +503,7 @@ void UText3DComponent::SetKerning(const float Value)
 
 void UText3DComponent::SetLineSpacing(const float Value)
 {
-	if (LineSpacing != Value)
+	if (!FMath::IsNearlyEqual(LineSpacing, Value))
 	{
 		LineSpacing = Value;
 		Rebuild();
@@ -527,7 +512,7 @@ void UText3DComponent::SetLineSpacing(const float Value)
 
 void UText3DComponent::SetWordSpacing(const float Value)
 {
-	if (WordSpacing != Value)
+	if (!FMath::IsNearlyEqual(WordSpacing, Value))
 	{
 		WordSpacing = Value;
 		Rebuild();
@@ -555,7 +540,7 @@ void UText3DComponent::SetVerticalAlignment(const EText3DVerticalTextAlignment V
 void UText3DComponent::SetExtrude(const float Value)
 {
 	const float NewValue = FMath::Max(0.0f, Value);
-	if (Extrude != NewValue)
+	if (!FMath::IsNearlyEqual(Extrude, NewValue))
 	{
 		Extrude = NewValue;
 		CheckBevel();
@@ -584,7 +569,7 @@ void UText3DComponent::SetHasMaxWidth(const bool Value)
 void UText3DComponent::SetMaxWidth(const float Value)
 {
 	const float NewValue = FMath::Max(1.0f, Value);
-	if (MaxWidth != NewValue)
+	if (!FMath::IsNearlyEqual(MaxWidth, NewValue))
 	{
 		MaxWidth = NewValue;
 		Rebuild();
@@ -603,7 +588,7 @@ void UText3DComponent::SetHasMaxHeight(const bool Value)
 void UText3DComponent::SetMaxHeight(const float Value)
 {
 	const float NewValue = FMath::Max(1.0f, Value);
-	if (MaxHeight != NewValue)
+	if (!FMath::IsNearlyEqual(MaxHeight, NewValue))
 	{
 		MaxHeight = NewValue;
 		Rebuild();
@@ -623,7 +608,7 @@ void UText3DComponent::SetBevel(const float Value)
 {
 	const float NewValue = FMath::Clamp(Value, 0.f, MaxBevel());
 
-	if (Bevel != NewValue)
+	if (!FMath::IsNearlyEqual(Bevel, NewValue))
 	{
 		Bevel = NewValue;
 		Rebuild();
@@ -646,7 +631,7 @@ void UText3DComponent::SetHalfCircleSegments(const int32 Value)
 		return;
 	}
 
-	const int NewValue = FMath::Clamp(Value, 1, 10);
+	const int32 NewValue = FMath::Clamp(Value, 1, 10);
 
 	if (HalfCircleSegments != NewValue)
 	{
@@ -657,12 +642,12 @@ void UText3DComponent::SetHalfCircleSegments(const int32 Value)
 
 void UText3DComponent::SetFreeze(const bool bFreeze)
 {
-	FreezeBuild = bFreeze;
+	bFreezeBuild = bFreeze;
 	if (bFreeze)
 	{
-		PendingBuild = false;
+		bPendingBuild = false;
 	}
-	else if (PendingBuild)
+	else if (bPendingBuild)
 	{
 		Rebuild();
 	}
@@ -673,7 +658,7 @@ void UText3DComponent::SetFrontMaterial(UMaterialInterface* Value)
 	if (Value != FrontMaterial)
 	{
 		FrontMaterial = Value;
-		if (!FreezeBuild)
+		if (!bFreezeBuild)
 		{
 			MarkRenderStateDirty();
 		}
@@ -685,7 +670,7 @@ void UText3DComponent::SetBevelMaterial(UMaterialInterface* Value)
 	if (Value != BevelMaterial)
 	{
 		BevelMaterial = Value;
-		if (!FreezeBuild)
+		if (!bFreezeBuild)
 		{
 			MarkRenderStateDirty();
 		}
@@ -697,7 +682,7 @@ void UText3DComponent::SetExtrudeMaterial(UMaterialInterface* Value)
 	if (Value != ExtrudeMaterial)
 	{
 		ExtrudeMaterial = Value;
-		if (!FreezeBuild)
+		if (!bFreezeBuild)
 		{
 			MarkRenderStateDirty();
 		}
@@ -709,7 +694,7 @@ void UText3DComponent::SetBackMaterial(UMaterialInterface* Value)
 	if (Value != BackMaterial)
 	{
 		BackMaterial = Value;
-		if (!FreezeBuild)
+		if (!bFreezeBuild)
 		{
 			MarkRenderStateDirty();
 		}
@@ -718,8 +703,8 @@ void UText3DComponent::SetBackMaterial(UMaterialInterface* Value)
 
 void UText3DComponent::Rebuild()
 {
-	PendingBuild = true;
-	if (!FreezeBuild)
+	bPendingBuild = true;
+	if (!bFreezeBuild)
 	{
 		MarkRenderStateDirty();
 	}
@@ -837,6 +822,9 @@ void UText3DComponent::BuildTextMesh()
 	const TSharedPtr<FData> MeshesData = MakeShared<FData>(Meshes, Bevel, FontInverseScale, Scale);
 	MeshesData->SetVerticalOffset(VerticalOffset);
 
+	FMeshCreator MeshCreator(Meshes, MeshesData);
+
+
 	for (const FShapedGlyphLine& ShapedLine : ShapedLines)
 	{
 		float HorizontalOffset = 0.0f;
@@ -855,96 +843,16 @@ void UText3DComponent::BuildTextMesh()
 		{
 			if (ShapedGlyph.bIsVisible)
 			{
-				MeshesData->SetCurrentMesh(EText3DMeshType::Front);
-				MeshesData->ResetDoneExtrude();
-
 				if (FT_Load_Glyph(Face, ShapedGlyph.GlyphIndex, FT_LOAD_DEFAULT))
 				{
 					continue;
 				}
 
-				FTVectoriser Vectoriser(Face->glyph);
-				Vectoriser.MakeMesh(FTGL_BACK_FACING);
-				const FTMesh* const Mesh = Vectoriser.GetMesh();
-				const size_t TesselateCount = Mesh->TesselationCount();
-
-				for (size_t i = 0; i < TesselateCount; i++)
-				{
-					const FTTesselation* const Tessselate = Mesh->Tesselation(i);
-					const int32 PointCount = Tessselate->PointCount();
-
-					if (PointCount < 3)
-					{
-						continue;
-					}
-
-					MeshesData->SetMinBevelTarget();
-					const int32 StartVertex = MeshesData->AddVertices(PointCount);
-
-					for (int32 PointIndex = 0; PointIndex < PointCount; PointIndex++)
-					{
-						const FTPoint* const Point = &Tessselate->Point(PointIndex);
-						MeshesData->AddVertex({Point->Xf(), Point->Yf()}, {1, 0}, {0, 0, -1});
-					}
-
-					switch (Tessselate->PolygonType())
-					{
-					case GL_TRIANGLES:
-					{
-						const int32 TrianglesCount = PointCount / 3;
-						MeshesData->AddTriangles(TrianglesCount);
-
-						for (int32 Index = 0; Index < TrianglesCount; Index++)
-						{
-							const int32 TriangleStart = StartVertex + Index * 3;
-							MeshesData->AddTriangle(TriangleStart + 0, TriangleStart + 2, TriangleStart + 1);
-						}
-
-						break;
-					}
-					case GL_TRIANGLE_STRIP:
-					{
-						MeshesData->AddTriangles(PointCount - 2);
-
-						for (int32 Index = 0; Index < PointCount - 2; Index++)
-						{
-							const int32 TriangleStart = StartVertex + Index;
-							const int32 Offset = (Index % 2) ? 0 : 1;
-							MeshesData->AddTriangle(TriangleStart + Offset, TriangleStart + 1 - Offset, TriangleStart + 2);
-						}
-
-						break;
-					}
-					case GL_TRIANGLE_FAN:
-					{
-						MeshesData->AddTriangles(PointCount - 2);
-
-						for (int32 Index = 1; Index < PointCount - 1; Index++)
-						{
-							const int32 TriangleStart = StartVertex + Index;
-							MeshesData->AddTriangle(StartVertex, TriangleStart + 1, TriangleStart + 0);
-						}
-
-						break;
-					}
-
-					default:
-					{
-						UE_LOG(LogText3D, Error, TEXT("Unknown polygon type: %X"), Tessselate->PolygonType());
-						break;
-					}
-					}
-				}
+				const TSharedPtr<FContourList> Contours = MakeShared<FContourList>(Face->glyph, MeshesData);
 
-				if (Extrude > 0.0f)
+				if (Contours->Num() != 0)
 				{
-					FText3DDebug Debugger;
-#if WITH_EDITOR
-					Debugger = DebugVariables;
-#endif
-
-					BevelContours(MeshesData, Vectoriser, Extrude, Bevel, BevelType, HalfCircleSegments,
-						Debugger.Iterations, Debugger.bHidePrevious, Debugger.MarkedVertex, Debugger.Segments, Debugger.VisibleFace);
+					MeshCreator.CreateMeshes(Contours, Extrude, Bevel, BevelType, HalfCircleSegments);
 				}
 			}
 
@@ -959,99 +867,8 @@ void UText3DComponent::BuildTextMesh()
 	FT_Done_Face(Face);
 	Face = nullptr;
 
-
-	for (int32 Type = 0; Type < static_cast<int32>(EText3DMeshType::TypeCount); Type++)
-	{
-		MeshesData->SetCurrentMesh(static_cast<EText3DMeshType>(Type));
-	}
-
-	struct FBoundingBox
-	{
-		FBox2D Box;
-		FVector2D Size;
-	};
-
-	TArray<FBoundingBox> BoundingBoxes;
-	FVector2D MaxSize(0, 0);
-	{
-		EText3DMeshType MeshType = FMath::IsNearlyZero(Bevel) ? EText3DMeshType::Front : EText3DMeshType::Bevel;
-		const FText3DMesh& Mesh = Meshes.Get()[static_cast<int32>(MeshType)];
-		const TArray<int32>& GlyphStartVertices = Mesh.GlyphStartVertices;
-		BoundingBoxes.AddUninitialized(GlyphStartVertices.Num() - 1);
-
-		for (int32 GlyphIndex = 0; GlyphIndex < BoundingBoxes.Num(); GlyphIndex++)
-		{
-			FBoundingBox& BoundingBox = BoundingBoxes[GlyphIndex];
-			FBox2D& Box = BoundingBox.Box;
-
-			const int32 FirstIndex = GlyphStartVertices[GlyphIndex];
-			const int32 LastIndex = GlyphStartVertices[GlyphIndex + 1];
-
-			if (FirstIndex < LastIndex)
-			{
-				const FVector& Position = Mesh.Vertices[FirstIndex].Position;
-				const FVector2D PositionFlat = {Position.Y, Position.Z};
-
-				Box.Min = PositionFlat;
-				Box.Max = PositionFlat;
-			}
-
-			for (int32 VertexIndex = FirstIndex + 1; VertexIndex < LastIndex; VertexIndex++)
-			{
-				const FDynamicMeshVertex& Vertex = Mesh.Vertices[VertexIndex];
-				const FVector& Position = Vertex.Position;
-
-				Box.Min.X = FMath::Min(Box.Min.X, Position.Y);
-				Box.Min.Y = FMath::Min(Box.Min.Y, Position.Z);
-				Box.Max.X = FMath::Max(Box.Max.X, Position.Y);
-				Box.Max.Y = FMath::Max(Box.Max.Y, Position.Z);
-			}
-
-			FVector2D& Size = BoundingBox.Size;
-			Size = Box.GetSize();
-
-			MaxSize.X = FMath::Max(MaxSize.X, Size.X);
-			MaxSize.Y = FMath::Max(MaxSize.Y, Size.Y);
-		}
-	}
-
-	for (int32 GlyphIndex = 0; GlyphIndex < BoundingBoxes.Num(); GlyphIndex++)
-	{
-		const int32 NextGlyphIndex = GlyphIndex + 1;
-		FBoundingBox& BoundingBox = BoundingBoxes[GlyphIndex];
-
-		TSharedPtr<TText3DMeshList> MeshesLocal = Meshes;
-		auto SetTextureCoordinates = [MeshesLocal, NextGlyphIndex, GlyphIndex, &BoundingBox, MaxSize](const EText3DMeshType Mesh)
-		{
-			FText3DMesh& CurrentMesh = (*MeshesLocal.Get())[static_cast<int32>(Mesh)];
-			const TArray<int32>& GlyphStartVertices = CurrentMesh.GlyphStartVertices;
-
-			if (NextGlyphIndex >= GlyphStartVertices.Num())
-			{
-				return;
-			}
-
-			for (int32 VertexIndex = GlyphStartVertices[GlyphIndex]; VertexIndex < GlyphStartVertices[NextGlyphIndex]; VertexIndex++)
-			{
-				FDynamicMeshVertex& Vertex = CurrentMesh.Vertices[VertexIndex];
-				const FVector2D TextureCoordinate = (FVector2D(Vertex.Position.Y, Vertex.Position.Z) - BoundingBox.Box.Min) / MaxSize;
-
-				for (int32 Index = 0; Index < MAX_STATIC_TEXCOORDS; Index++)
-				{
-					Vertex.TextureCoordinate[Index] = {TextureCoordinate.X, 1 - TextureCoordinate.Y};
-				}
-			}
-		};
-
-		SetTextureCoordinates(EText3DMeshType::Front);
-		SetTextureCoordinates(EText3DMeshType::Bevel);
-		SetTextureCoordinates(EText3DMeshType::Back);
-	}
-
-
-	// When TEXT3D_WITH_INTERSECTION=1 the following functions will not do anything
-	MirrorMesh(EText3DMeshType::Bevel, EText3DMeshType::Bevel, Scale.X);
-	MirrorMesh(EText3DMeshType::Front, EText3DMeshType::Back, Scale.X);
+	MeshCreator.SetFrontAndBevelTextureCoordinates(Bevel);
+	MeshCreator.MirrorMeshes(Extrude, Scale.X);
 }
 
 void UText3DComponent::CheckBevel()
@@ -1071,60 +888,13 @@ float UText3DComponent::MaxBevel() const
 #endif
 }
 
-void UText3DComponent::MirrorMesh(const EText3DMeshType TypeIn, const EText3DMeshType TypeOut, const float ScaleX)
-{
-#if !TEXT3D_WITH_INTERSECTION
-	const FText3DMesh& MeshIn = Meshes.Get()[static_cast<int32>(TypeIn)];
-	FText3DMesh& MeshOut = Meshes.Get()[static_cast<int32>(TypeOut)];
-
-
-	const TArray<FDynamicMeshVertex>& VerticesIn = MeshIn.Vertices;
-	TArray<FDynamicMeshVertex>& VerticesOut = MeshOut.Vertices;
-
-	const int32 VerticesInNum = VerticesIn.Num();
-	const int32 VerticesOutNum = VerticesOut.Num();
-
-	VerticesOut.AddUninitialized(VerticesInNum);
-
-	for (int32 Index = 0; Index < VerticesInNum; Index++)
-	{
-		const FDynamicMeshVertex& Vertex = VerticesIn[Index];
-
-		const FVector& Position = Vertex.Position;
-		const FVector TangentX = Vertex.TangentX.ToFVector();
-		const FVector TangentZ = Vertex.TangentZ.ToFVector();
-
-		VerticesOut[VerticesOutNum + Index] = FDynamicMeshVertex({Extrude * ScaleX - Position.X, Position.Y, Position.Z}, {-TangentX.X, TangentX.Y, TangentX.Z}, {-TangentZ.X, TangentZ.Y, TangentZ.Z}, Vertex.TextureCoordinate[0], Vertex.Color);
-	}
-
-
-	const TArray<int32>& IndicesIn = MeshIn.Indices;
-	TArray<int32>& IndicesOut = MeshOut.Indices;
-
-	const int32 IndicesInNum = IndicesIn.Num();
-	const int32 IndicesOutNum = IndicesOut.Num();
-
-	IndicesOut.AddUninitialized(IndicesInNum);
-
-	for (int32 Index = 0; Index < IndicesInNum / 3; Index++)
-	{
-		const int32 OldTriangle = Index * 3;
-		const int32 NewTriangle = IndicesOutNum + OldTriangle;
-
-		IndicesOut[NewTriangle + 0] = VerticesOutNum + IndicesIn[OldTriangle + 0];
-		IndicesOut[NewTriangle + 1] = VerticesOutNum + IndicesIn[OldTriangle + 2];
-		IndicesOut[NewTriangle + 2] = VerticesOutNum + IndicesIn[OldTriangle + 1];
-	}
-#endif
-}
-
 void UText3DComponent::OnRegister()
 {
 	CheckBevel();
 	BuildTextMesh();
 	Super::OnRegister();
 
-	if (!FreezeBuild)
+	if (!bFreezeBuild)
 	{
 		MarkRenderStateDirty();
 	}
@@ -1145,12 +915,12 @@ void UText3DComponent::SendRenderDynamicData_Concurrent()
 		return;
 	}
 
-	if (PendingBuild)
+	if (bPendingBuild)
 	{
 		BuildTextMesh();
 		UpdateBounds();
 		TextSceneProxy->UpdateData();
-		PendingBuild = false;
+		bPendingBuild = false;
 	}
 
 
@@ -1175,10 +945,10 @@ void UText3DComponent::SendRenderDynamicData_Concurrent()
 
 	// Enqueue command to send to render thread
 	ENQUEUE_RENDER_COMMAND(FSendText3DDynamicData)(
-		[TextSceneProxy, d{ MoveTemp(DynamicData) }](FRHICommandListImmediate& RHICmdList)
-		{
-			TextSceneProxy->SetDynamicData_RenderThread(d);
-		});
+				[TextSceneProxy, d{ MoveTemp(DynamicData) }](FRHICommandListImmediate& RHICmdList)
+	{
+		TextSceneProxy->SetDynamicData_RenderThread(d);
+	});
 }
 
 FBoxSphereBounds UText3DComponent::CalcBounds(const FTransform& LocalToWorld) const
diff --git a/Engine/Plugins/Experimental/Text3D/Source/Text3D/Private/Text3DPrivate.h b/Engine/Plugins/Experimental/Text3D/Source/Text3D/Private/Text3DPrivate.h
--- a/Engine/Plugins/Experimental/Text3D/Source/Text3D/Private/Text3DPrivate.h
+++ b/Engine/Plugins/Experimental/Text3D/Source/Text3D/Private/Text3DPrivate.h
@@ -15,7 +15,7 @@ const float FontInverseScale = 1.0f / FontSize;
 
 DECLARE_LOG_CATEGORY_EXTERN(LogText3D, Log, All);
 
-class FText3DModule : public IModuleInterface
+class FText3DModule final : public IModuleInterface
 {
 public:
 	FText3DModule();
diff --git a/Engine/Plugins/Experimental/Text3D/Source/Text3D/Private/TextShaper.h b/Engine/Plugins/Experimental/Text3D/Source/Text3D/Private/TextShaper.h
--- a/Engine/Plugins/Experimental/Text3D/Source/Text3D/Private/TextShaper.h
+++ b/Engine/Plugins/Experimental/Text3D/Source/Text3D/Private/TextShaper.h
@@ -36,7 +36,7 @@ struct FShapedGlyphLine
 	}
 };
 
-class FTextShaper
+class FTextShaper final
 {
 public:
 	static FTextShaper * Get()							{ return Instance; }
diff --git a/Engine/Plugins/Experimental/Text3D/Source/Text3D/Private/Util.h b/Engine/Plugins/Experimental/Text3D/Source/Text3D/Private/Util.h
new file mode 100644
--- /dev/null
+++ b/Engine/Plugins/Experimental/Text3D/Source/Text3D/Private/Util.h
@@ -0,0 +1,24 @@
+// Copyright 1998-2019 Epic Games, Inc. All Rights Reserved.
+
+#pragma once
+#include "CoreMinimal.h"
+
+struct FPart;
+
+class FAvailableExpandsFar final : public TMap<FPart*, float>
+{
+public:
+	void Add(const FPart* const Edge, const float Value)
+	{
+		TMap<FPart*, float>::Add(const_cast<FPart*>(Edge), Value);
+	}
+};
+
+class FContourBase : public TArray<FPart*>
+{
+protected:
+	int32 Find(const FPart* const Edge)
+	{
+		return TArray<FPart*>::Find(const_cast<FPart*>(Edge));
+	}
+};
diff --git a/Engine/Plugins/Experimental/Text3D/Source/Text3D/Public/BevelType.h b/Engine/Plugins/Experimental/Text3D/Source/Text3D/Public/BevelType.h
new file mode 100644
--- /dev/null
+++ b/Engine/Plugins/Experimental/Text3D/Source/Text3D/Public/BevelType.h
@@ -0,0 +1,12 @@
+// Copyright 1998-2019 Epic Games, Inc. All Rights Reserved.
+
+#pragma once
+
+#include "UObject/ObjectMacros.h"
+
+UENUM()
+enum class EText3DBevelType : uint8
+{
+    Linear,
+    HalfCircle
+};
diff --git a/Engine/Plugins/Experimental/Text3D/Source/Text3D/Public/Mesh.h b/Engine/Plugins/Experimental/Text3D/Source/Text3D/Public/Mesh.h
new file mode 100644
--- /dev/null
+++ b/Engine/Plugins/Experimental/Text3D/Source/Text3D/Public/Mesh.h
@@ -0,0 +1,41 @@
+// Copyright 1998-2019 Epic Games, Inc. All Rights Reserved.
+
+#pragma once
+
+#include "CoreMinimal.h"
+#include "DynamicMeshBuilder.h"
+
+
+struct FText3DDynamicData
+{
+	FText3DDynamicData()
+	{
+	}
+
+	FText3DDynamicData(const TArray<int32>& InIndices, const TArray<FDynamicMeshVertex>& InVertices) :
+		Indices(InIndices),
+		Vertices(InVertices)
+	{
+	}
+
+	TArray<int32> Indices;
+	TArray<FDynamicMeshVertex> Vertices;
+	TArray<int32> GlyphStartVertices;
+};
+
+struct FText3DMesh : public FText3DDynamicData
+{
+	bool IsEmpty() const;
+};
+
+enum class EText3DMeshType : uint32
+{
+	Front = 0,
+	Bevel = 1,
+	Extrude = 2,
+	Back = 3,
+
+	TypeCount = 4
+};
+
+using TText3DMeshList = TArray<FText3DMesh, TFixedAllocator<static_cast<int32>(EText3DMeshType::TypeCount)>>;
\ No newline at end of file
diff --git a/Engine/Plugins/Experimental/Text3D/Source/Text3D/Public/Text3DActor.h b/Engine/Plugins/Experimental/Text3D/Source/Text3D/Public/Text3DActor.h
--- a/Engine/Plugins/Experimental/Text3D/Source/Text3D/Public/Text3DActor.h
+++ b/Engine/Plugins/Experimental/Text3D/Source/Text3D/Public/Text3DActor.h
@@ -9,7 +9,7 @@
 class UText3DComponent;
 
 UCLASS(Meta = (DisplayName = "Text 3D", ComponentWrapperClass))
-class TEXT3D_API AText3DActor : public AActor
+class TEXT3D_API AText3DActor final : public AActor
 {
 	GENERATED_BODY()
 
diff --git a/Engine/Plugins/Experimental/Text3D/Source/Text3D/Public/Text3DComponent.h b/Engine/Plugins/Experimental/Text3D/Source/Text3D/Public/Text3DComponent.h
--- a/Engine/Plugins/Experimental/Text3D/Source/Text3D/Public/Text3DComponent.h
+++ b/Engine/Plugins/Experimental/Text3D/Source/Text3D/Public/Text3DComponent.h
@@ -3,8 +3,8 @@
 #pragma once
 
 
-#include "Bevel/Mesh.h"
-#include "Bevel/BevelType.h"
+#include "Mesh.h"
+#include "BevelType.h"
 
 #include "CoreMinimal.h"
 #include "Components/PrimitiveComponent.h"
@@ -35,39 +35,8 @@ enum class EText3DHorizontalTextAlignment : uint8
 	Right			UMETA(DisplayName = "Right"),
 };
 
-USTRUCT()
-struct FText3DDebug
-{
-public:
-	GENERATED_BODY()
-
-	FText3DDebug()
-	{
-		Iterations = 1;
-		bHidePrevious = false;
-		MarkedVertex = -1;
-		Segments = 12;
-		VisibleFace = -1;
-	}
-
-	UPROPERTY(EditAnywhere, Category = "Text3D Debug")
-	int32 Iterations;
-
-	UPROPERTY(EditAnywhere, Category = "Text3D Debug")
-	bool bHidePrevious;
-
-	UPROPERTY(EditAnywhere, Category = "Text3D Debug")
-	int32 MarkedVertex;
-
-	UPROPERTY(EditAnywhere, Category = "Text3D Debug")
-	int32 Segments;
-
-	UPROPERTY(EditAnywhere, Category = "Text3D Debug")
-	int32 VisibleFace;
-};
-
 UCLASS(HideCategories=(Object, Physics, Activation, "Components|Activation"), EditInlineNew, meta=(BlueprintSpawnableComponent))
-class TEXT3D_API UText3DComponent : public UPrimitiveComponent
+class TEXT3D_API UText3DComponent final : public UPrimitiveComponent
 {
 	GENERATED_BODY()
 
@@ -103,7 +72,7 @@ public:
 	EText3DBevelType BevelType;
 
 	/** Half Circle Bevel Segments (Defines the amount of tesselation for the bevel part) */
-	UPROPERTY(EditAnywhere, BlueprintReadWrite, BlueprintSetter=SetHalfCircleSegments, Category="Text3D", Meta=(ClampMin=1, ClampMax=10))
+	UPROPERTY(EditAnywhere, BlueprintReadWrite, BlueprintSetter=SetHalfCircleSegments, Category="Text3D", Meta=(ClampMin=1, ClampMax=15))
 	int32 HalfCircleSegments;
 
 	/** Material for the front part */
@@ -250,14 +219,9 @@ public:
 	UFUNCTION(BlueprintCallable, Category = "Rendering|Components|Text3D")
 	void SetFreeze(const bool bFreeze);
 
-#if WITH_EDITORONLY_DATA
-	UPROPERTY(EditAnywhere, Category = "Text3D Debug")
-	FText3DDebug DebugVariables;
-#endif
-
 private:
-	bool PendingBuild;
-	bool FreezeBuild;
+	bool bPendingBuild;
+	bool bFreezeBuild;
 
 	TSharedRef<TText3DMeshList> Meshes;
 	FTransform CharacterTransform;
@@ -268,5 +232,4 @@ private:
 	void BuildTextMesh();
 	void CheckBevel();
 	float MaxBevel() const;
-	void MirrorMesh(const EText3DMeshType TypeIn, const EText3DMeshType TypeOut, const float ScaleX);
 };
diff --git a/Engine/Plugins/Experimental/Text3D/Source/Text3D/Text3D.Build.cs b/Engine/Plugins/Experimental/Text3D/Source/Text3D/Text3D.Build.cs
--- a/Engine/Plugins/Experimental/Text3D/Source/Text3D/Text3D.Build.cs
+++ b/Engine/Plugins/Experimental/Text3D/Source/Text3D/Text3D.Build.cs
@@ -26,7 +26,6 @@ public class Text3D : ModuleRules
 			// 3rd party libraries
 			"FreeType2",
 			"OpenGL",   // only for include purposes
-			"FTGL",
 			"HarfBuzz",
 			"GLUtesselator"
 		});
diff --git a/Engine/Source/Developer/Datasmith/DatasmithFacade/Private/DatasmithFacadeMaterial.cpp b/Engine/Source/Developer/Datasmith/DatasmithFacade/Private/DatasmithFacadeMaterial.cpp
--- a/Engine/Source/Developer/Datasmith/DatasmithFacade/Private/DatasmithFacadeMaterial.cpp
+++ b/Engine/Source/Developer/Datasmith/DatasmithFacade/Private/DatasmithFacadeMaterial.cpp
@@ -1,6 +1,7 @@
 // Copyright 1998-2019 Epic Games, Inc. All Rights Reserved.
 
 #include "DatasmithFacadeMaterial.h"
+#include "DatasmithUtils.h"
 
 #include "Misc/Paths.h"
 
@@ -64,8 +65,9 @@ void FDatasmithFacadeMaterial::AddTexture(
 {
 	if (!FString(InTextureFilePath).IsEmpty())
 	{
-		// Make a unique Datasmith texture name.
-		FString TextureName = FString::Printf(TEXT("%ls_%d"), *FMD5::HashAnsiString(InTextureFilePath), int(InTextureMode));
+		// Make Datasmith texture name from file name.
+		FString FileName = FDatasmithUtils::SanitizeFileName(FPaths::GetBaseFilename(InTextureFilePath));
+		FString TextureName = FString::Printf(TEXT("%ls_%d"), *FileName, int(InTextureMode));
 
 		// Create a new Datasmith material texture property.
 		TSharedPtr<IDatasmithKeyValueProperty> MaterialPropertyPtr = FDatasmithSceneFactory::CreateKeyValueProperty(InPropertyName);
diff --git a/Engine/Source/Developer/RawMesh/Private/RawMesh.cpp b/Engine/Source/Developer/RawMesh/Private/RawMesh.cpp
--- a/Engine/Source/Developer/RawMesh/Private/RawMesh.cpp
+++ b/Engine/Source/Developer/RawMesh/Private/RawMesh.cpp
@@ -160,8 +160,6 @@ FRawMeshBulkData::FRawMeshBulkData()
 {
 }
 
-#if WITH_EDITORONLY_DATA
-
 /**
  * Serialization of raw meshes uses its own versioning scheme because it is
  * stored in bulk data.
@@ -320,6 +318,4 @@ void FRawMeshBulkData::Empty()
 	BulkData.RemoveBulkData();
 	Guid.Invalidate();
 	bGuidIsHash = false;
-}
-
-#endif // #if WITH_EDITORONLY_DATA
+}
\ No newline at end of file
diff --git a/Engine/Source/Developer/RawMesh/Public/RawMesh.h b/Engine/Source/Developer/RawMesh/Public/RawMesh.h
--- a/Engine/Source/Developer/RawMesh/Public/RawMesh.h
+++ b/Engine/Source/Developer/RawMesh/Public/RawMesh.h
@@ -110,7 +110,6 @@ public:
 	/** Default constructor. */
 	RAWMESH_API FRawMeshBulkData();
 
-#if WITH_EDITORONLY_DATA
 	/** Serialization. */
 	RAWMESH_API void Serialize(class FArchive& Ar, class UObject* Owner);
 
@@ -133,5 +132,4 @@ public:
 
 	/** Returns true if no bulk data is available for this mesh. */
 	RAWMESH_API void Empty();
-#endif // #if WITH_EDITORONLY_DATA
 };
diff --git a/Engine/Source/Editor/UnrealEd/Private/Fbx/FbxSkeletalMeshImport.cpp b/Engine/Source/Editor/UnrealEd/Private/Fbx/FbxSkeletalMeshImport.cpp
--- a/Engine/Source/Editor/UnrealEd/Private/Fbx/FbxSkeletalMeshImport.cpp
+++ b/Engine/Source/Editor/UnrealEd/Private/Fbx/FbxSkeletalMeshImport.cpp
@@ -3881,6 +3881,12 @@ void UnFbx::FFbxImporter::InsertNewLODToBaseSkeletalMesh(USkeletalMesh* InSkelet
 		if (LODMatIndex == INDEX_NONE && InSkeletalMesh->Materials[MatIdx].MaterialInterface != NULL)
 		{
 			LODMatIndex = BaseSkeletalMesh->Materials.Find(InSkeletalMesh->Materials[MatIdx]);
+			
+			//If the material does not exist in the base skeletal mesh but is used in other LODs, add it.
+			if (LODMatIndex == INDEX_NONE)
+			{
+				LODMatIndex = BaseSkeletalMesh->Materials.Add(InSkeletalMesh->Materials[MatIdx]);
+			}
 		}
 
 		// If we didn't just use the index - but make sure its within range of the Materials array.
diff --git a/Engine/Source/Programs/AutomationTool/Scripts/CheckCsprojDotNetVersion.cs b/Engine/Source/Programs/AutomationTool/Scripts/CheckCsprojDotNetVersion.cs
--- a/Engine/Source/Programs/AutomationTool/Scripts/CheckCsprojDotNetVersion.cs
+++ b/Engine/Source/Programs/AutomationTool/Scripts/CheckCsprojDotNetVersion.cs
@@ -37,7 +37,7 @@ class CheckCsprojDotNetVersion : BuildCommand
                 (CsProj.ContainsName("UE4TemplateProject", EngineDir) && CsProj.GetFileName().Equals("ProjectTemplate.csproj")) ||
                 CsProj.GetFileNameWithoutExtension().ToLower().Contains("_mono") ||
                 CsProj.GetFileNameWithoutExtension().ToLower().Contains("unrealvs") ||
-				CsProj.GetFileNameWithoutExtension().ToLower().Contains("revit"))
+				CsProj.ContainsName("DatasmithRevitExporter", EngineDir))
 
             {
                 continue;
diff --git a/Engine/Source/Programs/Enterprise/Datasmith/DatasmithCADWorker/Private/DatasmithCADWorker.cpp b/Engine/Source/Programs/Enterprise/Datasmith/DatasmithCADWorker/Private/DatasmithCADWorker.cpp
--- a/Engine/Source/Programs/Enterprise/Datasmith/DatasmithCADWorker/Private/DatasmithCADWorker.cpp
+++ b/Engine/Source/Programs/Enterprise/Datasmith/DatasmithCADWorker/Private/DatasmithCADWorker.cpp
@@ -37,13 +37,13 @@ bool HasParameter(int32 Argc, TCHAR* Argv[], const FString& InParam)
 
 int32 Main(int32 Argc, TCHAR * Argv[])
 {
-	FString ServerPID, ServerPort, CacheDirectory, KernelIOPath;
+	FString ServerPID, ServerPort, CacheDirectory, EnginePluginsPath;
 	GetParameter(Argc, Argv, "-ServerPID", ServerPID);
 	GetParameter(Argc, Argv, "-ServerPort", ServerPort);
 	GetParameter(Argc, Argv, "-CacheDir", CacheDirectory);
-	GetParameter(Argc, Argv, "-KernelIOPath", KernelIOPath);
+	GetParameter(Argc, Argv, "-EnginePluginsDir", EnginePluginsPath);
 
-	FDatasmithCADWorkerImpl Worker(FCString::Atoi(*ServerPID), FCString::Atoi(*ServerPort), KernelIOPath, CacheDirectory);
+	FDatasmithCADWorkerImpl Worker(FCString::Atoi(*ServerPID), FCString::Atoi(*ServerPort), EnginePluginsPath, CacheDirectory);
 	Worker.Run();
 
 	return 0;
diff --git a/Engine/Source/Programs/Enterprise/Datasmith/DatasmithCADWorker/Private/DatasmithCADWorkerImpl.cpp b/Engine/Source/Programs/Enterprise/Datasmith/DatasmithCADWorker/Private/DatasmithCADWorkerImpl.cpp
--- a/Engine/Source/Programs/Enterprise/Datasmith/DatasmithCADWorker/Private/DatasmithCADWorkerImpl.cpp
+++ b/Engine/Source/Programs/Enterprise/Datasmith/DatasmithCADWorker/Private/DatasmithCADWorkerImpl.cpp
@@ -16,10 +16,10 @@ DEFINE_LOG_CATEGORY(LogDatasmithCADWorker);
 using namespace DatasmithDispatcher;
 
 
-FDatasmithCADWorkerImpl::FDatasmithCADWorkerImpl(int32 InServerPID, int32 InServerPort, const FString& InKernelIOPath, const FString& InCachePath)
+FDatasmithCADWorkerImpl::FDatasmithCADWorkerImpl(int32 InServerPID, int32 InServerPort, const FString& InEnginePluginsPath, const FString& InCachePath)
 	: ServerPID(InServerPID)
 	, ServerPort(InServerPort)
-	, KernelIOPath(InKernelIOPath)
+	, EnginePluginsPath(InEnginePluginsPath)
 	, CachePath(InCachePath)
 	, PingStartCycle(0)
 {
@@ -126,8 +126,8 @@ void FDatasmithCADWorkerImpl::ProcessCommand(const FRunTaskCommand& RunTaskComma
 
 	FCompletedTaskCommand CompletedTask;
 #ifdef CAD_INTERFACE
-	CADLibrary::FCoreTechFileParser FileParser(FileToProcess, CachePath, ImportParameters, *KernelIOPath);
-	CADLibrary::FCoreTechFileParser::EProcessResult ProcessResult = FileParser.ProcessFile();
+	CADLibrary::FCoreTechFileParser FileParser(ImportParameters, EnginePluginsPath, CachePath);
+	CADLibrary::FCoreTechFileParser::EProcessResult ProcessResult = FileParser.ProcessFile(FileToProcess);
 
 	CompletedTask.ProcessResult = ProcessResult;
 
diff --git a/Engine/Source/Programs/Enterprise/Datasmith/DatasmithCADWorker/Private/DatasmithCADWorkerImpl.h b/Engine/Source/Programs/Enterprise/Datasmith/DatasmithCADWorker/Private/DatasmithCADWorkerImpl.h
--- a/Engine/Source/Programs/Enterprise/Datasmith/DatasmithCADWorker/Private/DatasmithCADWorkerImpl.h
+++ b/Engine/Source/Programs/Enterprise/Datasmith/DatasmithCADWorker/Private/DatasmithCADWorkerImpl.h
@@ -19,7 +19,7 @@ struct FImportParameters;
 class FDatasmithCADWorkerImpl
 {
 public:
-	FDatasmithCADWorkerImpl(int32 InServerPID, int32 InServerPort, const FString& KernelIOPath, const FString& InCachePath);
+	FDatasmithCADWorkerImpl(int32 InServerPID, int32 InServerPort, const FString& EnginePluginsPath, const FString& InCachePath);
 	bool Run();
 
 private:
@@ -36,7 +36,7 @@ private:
 
 	int32 ServerPID;
 	int32 ServerPort;
-	FString KernelIOPath;
+	FString EnginePluginsPath;
 	FString CachePath;
 	CADLibrary::FImportParameters ImportParameters;
 	uint64 PingStartCycle;
diff --git a/Engine/Source/Programs/Enterprise/Datasmith/DatasmithRevitExporter/DatasmithRevit2018/DatasmithDynamoNode.csproj b/Engine/Source/Programs/Enterprise/Datasmith/DatasmithRevitExporter/DatasmithRevit2018/DatasmithDynamoNode.csproj
new file mode 100644
--- /dev/null
+++ b/Engine/Source/Programs/Enterprise/Datasmith/DatasmithRevitExporter/DatasmithRevit2018/DatasmithDynamoNode.csproj
@@ -0,0 +1,68 @@
+﻿<?xml version="1.0" encoding="utf-8"?>
+<Project ToolsVersion="15.0" xmlns="http://schemas.microsoft.com/developer/msbuild/2003">
+  <Import Project="$(MSBuildExtensionsPath)\$(MSBuildToolsVersion)\Microsoft.Common.props" Condition="Exists('$(MSBuildExtensionsPath)\$(MSBuildToolsVersion)\Microsoft.Common.props')" />
+  <PropertyGroup>
+    <Configuration Condition=" '$(Configuration)' == '' ">Debug</Configuration>
+    <Platform Condition=" '$(Platform)' == '' ">AnyCPU</Platform>
+    <ProjectGuid>{144B1728-4071-4972-AD7C-C4ACE006BE59}</ProjectGuid>
+    <OutputType>Library</OutputType>
+    <AppDesignerFolder>Properties</AppDesignerFolder>
+    <RootNamespace>DatasmithDynamoNode</RootNamespace>
+    <AssemblyName>DatasmithDynamoNode</AssemblyName>
+    <TargetFrameworkVersion>v4.5.2</TargetFrameworkVersion>
+    <FileAlignment>512</FileAlignment>
+    <Deterministic>true</Deterministic>
+    <TargetFrameworkProfile />
+  </PropertyGroup>
+  <PropertyGroup Condition=" '$(Configuration)|$(Platform)' == 'Debug|AnyCPU' ">
+    <DebugSymbols>true</DebugSymbols>
+    <DebugType>full</DebugType>
+    <Optimize>false</Optimize>
+    <OutputPath>bin\Debug\</OutputPath>
+    <DefineConstants>TRACE;DEBUG;REVIT_API_2018</DefineConstants>
+    <ErrorReport>prompt</ErrorReport>
+    <WarningLevel>4</WarningLevel>
+  </PropertyGroup>
+  <PropertyGroup Condition=" '$(Configuration)|$(Platform)' == 'Release|AnyCPU' ">
+    <DebugType>pdbonly</DebugType>
+    <Optimize>true</Optimize>
+    <OutputPath>bin\Release\</OutputPath>
+    <DefineConstants>TRACE</DefineConstants>
+    <ErrorReport>prompt</ErrorReport>
+    <WarningLevel>4</WarningLevel>
+  </PropertyGroup>
+  <ItemGroup>
+    <Reference Include="RevitAPI">
+      <HintPath>$(Revit_2018_API)\RevitAPI.dll</HintPath>
+      <Private>False</Private>
+    </Reference>
+    <Reference Include="RevitAPIUI">
+      <HintPath>$(Revit_2018_API)\RevitAPIUI.dll</HintPath>
+      <Private>False</Private>
+    </Reference>
+    <Reference Include="System" />
+    <Reference Include="System.Core" />
+    <Reference Include="System.Xml.Linq" />
+    <Reference Include="System.Data.DataSetExtensions" />
+    <Reference Include="Microsoft.CSharp" />
+    <Reference Include="System.Data" />
+    <Reference Include="System.Net.Http" />
+    <Reference Include="System.Xml" />
+  </ItemGroup>
+  <ItemGroup>
+    <Compile Include="..\Private\DatasmithDynamoNode.cs" />
+    <Compile Include="Properties\AssemblyInfo.cs" />
+  </ItemGroup>
+  <ItemGroup>
+    <ProjectReference Include="DatasmithRevit2018.csproj">
+      <Project>{159217e1-3787-4785-ad49-5d3aaf0b7d85}</Project>
+      <Name>DatasmithRevit2018</Name>
+    </ProjectReference>
+  </ItemGroup>
+  <Import Project="$(MSBuildToolsPath)\Microsoft.CSharp.targets" />
+  <PropertyGroup>
+    <PostBuildEvent>set PUBLISH_DIR=$(ProjectDir)..\..\..\..\..\..\Binaries\Win64\Revit\2018\
+xcopy "$(ProjectDir)$(OutputPath)*.dll" "%25PUBLISH_DIR%25" /Y
+</PostBuildEvent>
+  </PropertyGroup>
+</Project>
\ No newline at end of file
diff --git a/Engine/Source/Programs/Enterprise/Datasmith/DatasmithRevitExporter/DatasmithRevit2018/DatasmithRevit2018.csproj b/Engine/Source/Programs/Enterprise/Datasmith/DatasmithRevitExporter/DatasmithRevit2018/DatasmithRevit2018.csproj
--- a/Engine/Source/Programs/Enterprise/Datasmith/DatasmithRevitExporter/DatasmithRevit2018/DatasmithRevit2018.csproj
+++ b/Engine/Source/Programs/Enterprise/Datasmith/DatasmithRevitExporter/DatasmithRevit2018/DatasmithRevit2018.csproj
@@ -92,7 +92,8 @@
 xcopy "$(ProjectDir)$(TargetName).addin" "%25PUBLISH_DIR%25" /Y
 xcopy "$(ProjectDir)$(OutputPath)*.dll" "%25PUBLISH_DIR%25" /Y
 xcopy "$(ProjectDir)..\Resources\Windows\DatasmithIcon32.png" "%25PUBLISH_DIR%25" /Y
-</PostBuildEvent>
+xcopy "$(ProjectDir)..\Resources\Windows\DatasmithRevitDynamoExport.dyn" "%25PUBLISH_DIR%25" /Y
+	</PostBuildEvent>
   </PropertyGroup>
   <!-- To modify your build process, add your task inside one of the targets below and uncomment it.
        Other similar extension points exist, see Microsoft.Common.targets.
diff --git a/Engine/Source/Programs/Enterprise/Datasmith/DatasmithRevitExporter/DatasmithRevit2019/DatasmithDynamoNode.csproj b/Engine/Source/Programs/Enterprise/Datasmith/DatasmithRevitExporter/DatasmithRevit2019/DatasmithDynamoNode.csproj
new file mode 100644
--- /dev/null
+++ b/Engine/Source/Programs/Enterprise/Datasmith/DatasmithRevitExporter/DatasmithRevit2019/DatasmithDynamoNode.csproj
@@ -0,0 +1,67 @@
+﻿<?xml version="1.0" encoding="utf-8"?>
+<Project ToolsVersion="15.0" xmlns="http://schemas.microsoft.com/developer/msbuild/2003">
+  <Import Project="$(MSBuildExtensionsPath)\$(MSBuildToolsVersion)\Microsoft.Common.props" Condition="Exists('$(MSBuildExtensionsPath)\$(MSBuildToolsVersion)\Microsoft.Common.props')" />
+  <PropertyGroup>
+    <Configuration Condition=" '$(Configuration)' == '' ">Debug</Configuration>
+    <Platform Condition=" '$(Platform)' == '' ">AnyCPU</Platform>
+    <ProjectGuid>{144B1728-4071-4972-AD7C-C4ACE006BE59}</ProjectGuid>
+    <OutputType>Library</OutputType>
+    <AppDesignerFolder>Properties</AppDesignerFolder>
+    <RootNamespace>DatasmithDynamoNode</RootNamespace>
+    <AssemblyName>DatasmithDynamoNode</AssemblyName>
+    <TargetFrameworkVersion>v4.7</TargetFrameworkVersion>
+    <FileAlignment>512</FileAlignment>
+    <Deterministic>true</Deterministic>
+  </PropertyGroup>
+  <PropertyGroup Condition=" '$(Configuration)|$(Platform)' == 'Debug|AnyCPU' ">
+    <DebugSymbols>true</DebugSymbols>
+    <DebugType>full</DebugType>
+    <Optimize>false</Optimize>
+    <OutputPath>bin\Debug\</OutputPath>
+    <DefineConstants>TRACE;DEBUG;REVIT_API_2019</DefineConstants>
+    <ErrorReport>prompt</ErrorReport>
+    <WarningLevel>4</WarningLevel>
+  </PropertyGroup>
+  <PropertyGroup Condition=" '$(Configuration)|$(Platform)' == 'Release|AnyCPU' ">
+    <DebugType>pdbonly</DebugType>
+    <Optimize>true</Optimize>
+    <OutputPath>bin\Release\</OutputPath>
+    <DefineConstants>TRACE</DefineConstants>
+    <ErrorReport>prompt</ErrorReport>
+    <WarningLevel>4</WarningLevel>
+  </PropertyGroup>
+  <ItemGroup>
+    <Reference Include="RevitAPI">
+      <HintPath>$(Revit_2019_API)\RevitAPI.dll</HintPath>
+      <Private>False</Private>
+    </Reference>
+    <Reference Include="RevitAPIUI">
+      <HintPath>$(Revit_2019_API)\RevitAPIUI.dll</HintPath>
+      <Private>False</Private>
+    </Reference>
+    <Reference Include="System" />
+    <Reference Include="System.Core" />
+    <Reference Include="System.Xml.Linq" />
+    <Reference Include="System.Data.DataSetExtensions" />
+    <Reference Include="Microsoft.CSharp" />
+    <Reference Include="System.Data" />
+    <Reference Include="System.Net.Http" />
+    <Reference Include="System.Xml" />
+  </ItemGroup>
+  <ItemGroup>
+    <Compile Include="..\Private\DatasmithDynamoNode.cs" />
+    <Compile Include="Properties\AssemblyInfo.cs" />
+  </ItemGroup>
+  <ItemGroup>
+    <ProjectReference Include="DatasmithRevit2019.csproj">
+      <Project>{702fe407-50fb-46fa-982d-b89624fa8a30}</Project>
+      <Name>DatasmithRevit2019</Name>
+    </ProjectReference>
+  </ItemGroup>
+  <Import Project="$(MSBuildToolsPath)\Microsoft.CSharp.targets" />
+  <PropertyGroup>
+    <PostBuildEvent>set PUBLISH_DIR=$(ProjectDir)..\..\..\..\..\..\Binaries\Win64\Revit\2019\
+xcopy "$(ProjectDir)$(OutputPath)*.dll" "%25PUBLISH_DIR%25" /Y
+</PostBuildEvent>
+  </PropertyGroup>
+</Project>
\ No newline at end of file
diff --git a/Engine/Source/Programs/Enterprise/Datasmith/DatasmithRevitExporter/DatasmithRevit2019/DatasmithRevit2019.csproj b/Engine/Source/Programs/Enterprise/Datasmith/DatasmithRevitExporter/DatasmithRevit2019/DatasmithRevit2019.csproj
--- a/Engine/Source/Programs/Enterprise/Datasmith/DatasmithRevitExporter/DatasmithRevit2019/DatasmithRevit2019.csproj
+++ b/Engine/Source/Programs/Enterprise/Datasmith/DatasmithRevitExporter/DatasmithRevit2019/DatasmithRevit2019.csproj
@@ -92,6 +92,7 @@
 xcopy "$(ProjectDir)$(TargetName).addin" "%25PUBLISH_DIR%25" /Y
 xcopy "$(ProjectDir)$(OutputPath)*.dll" "%25PUBLISH_DIR%25" /Y
 xcopy "$(ProjectDir)..\Resources\Windows\DatasmithIcon32.png" "%25PUBLISH_DIR%25" /Y
+xcopy "$(ProjectDir)..\Resources\Windows\DatasmithRevitDynamoExport.dyn" "%25PUBLISH_DIR%25" /Y
 </PostBuildEvent>
   </PropertyGroup>
   <!-- To modify your build process, add your task inside one of the targets below and uncomment it.
diff --git a/Engine/Source/Programs/Enterprise/Datasmith/DatasmithRevitExporter/DatasmithRevit2020/DatasmithDynamoNode.csproj b/Engine/Source/Programs/Enterprise/Datasmith/DatasmithRevitExporter/DatasmithRevit2020/DatasmithDynamoNode.csproj
new file mode 100644
--- /dev/null
+++ b/Engine/Source/Programs/Enterprise/Datasmith/DatasmithRevitExporter/DatasmithRevit2020/DatasmithDynamoNode.csproj
@@ -0,0 +1,68 @@
+﻿<?xml version="1.0" encoding="utf-8"?>
+<Project ToolsVersion="15.0" xmlns="http://schemas.microsoft.com/developer/msbuild/2003">
+  <Import Project="$(MSBuildExtensionsPath)\$(MSBuildToolsVersion)\Microsoft.Common.props" Condition="Exists('$(MSBuildExtensionsPath)\$(MSBuildToolsVersion)\Microsoft.Common.props')" />
+  <PropertyGroup>
+    <Configuration Condition=" '$(Configuration)' == '' ">Debug</Configuration>
+    <Platform Condition=" '$(Platform)' == '' ">AnyCPU</Platform>
+    <ProjectGuid>{144B1728-4071-4972-AD7C-C4ACE006BE59}</ProjectGuid>
+    <OutputType>Library</OutputType>
+    <AppDesignerFolder>Properties</AppDesignerFolder>
+    <RootNamespace>DatasmithDynamoNode</RootNamespace>
+    <AssemblyName>DatasmithDynamoNode</AssemblyName>
+    <TargetFrameworkVersion>v4.7</TargetFrameworkVersion>
+    <FileAlignment>512</FileAlignment>
+    <Deterministic>true</Deterministic>
+  </PropertyGroup>
+  <PropertyGroup Condition=" '$(Configuration)|$(Platform)' == 'Debug|AnyCPU' ">
+    <DebugSymbols>true</DebugSymbols>
+    <DebugType>full</DebugType>
+    <Optimize>false</Optimize>
+    <OutputPath>bin\Debug\</OutputPath>
+    <DefineConstants>TRACE;DEBUG;REVIT_API_2020</DefineConstants>
+    <ErrorReport>prompt</ErrorReport>
+    <WarningLevel>4</WarningLevel>
+  </PropertyGroup>
+  <PropertyGroup Condition=" '$(Configuration)|$(Platform)' == 'Release|AnyCPU' ">
+    <DebugType>pdbonly</DebugType>
+    <Optimize>true</Optimize>
+    <OutputPath>bin\Release\</OutputPath>
+    <DefineConstants>TRACE</DefineConstants>
+    <ErrorReport>prompt</ErrorReport>
+    <WarningLevel>4</WarningLevel>
+  </PropertyGroup>
+  <ItemGroup>
+    <Reference Include="RevitAPI">
+      <HintPath>$(Revit_2020_API)\RevitAPI.dll</HintPath>
+      <Private>False</Private>
+    </Reference>
+    <Reference Include="RevitAPIUI">
+      <HintPath>$(Revit_2020_API)\RevitAPIUI.dll</HintPath>
+      <Private>False</Private>
+    </Reference>
+    <Reference Include="System" />
+    <Reference Include="System.Core" />
+    <Reference Include="System.Xml.Linq" />
+    <Reference Include="System.Data.DataSetExtensions" />
+    <Reference Include="Microsoft.CSharp" />
+    <Reference Include="System.Data" />
+    <Reference Include="System.Net.Http" />
+    <Reference Include="System.Xml" />
+  </ItemGroup>
+  <ItemGroup>
+    <Compile Include="..\Private\DatasmithDynamoNode.cs" />
+    <Compile Include="Properties\AssemblyInfo.cs" />
+  </ItemGroup>
+  <ItemGroup>
+    <ProjectReference Include="DatasmithRevit2020.csproj">
+      <Project>{702fe407-50fb-46fa-982d-b89624fa8a30}</Project>
+      <Name>DatasmithRevit2020</Name>
+      <Private>False</Private>
+    </ProjectReference>
+  </ItemGroup>
+  <Import Project="$(MSBuildToolsPath)\Microsoft.CSharp.targets" />
+  <PropertyGroup>
+    <PostBuildEvent>set PUBLISH_DIR=$(ProjectDir)..\..\..\..\..\..\Binaries\Win64\Revit\2020\
+xcopy "$(ProjectDir)$(OutputPath)*.dll" "%25PUBLISH_DIR%25" /Y
+</PostBuildEvent>
+  </PropertyGroup>
+</Project>
\ No newline at end of file
diff --git a/Engine/Source/Programs/Enterprise/Datasmith/DatasmithRevitExporter/DatasmithRevit2020/DatasmithRevit2020.csproj b/Engine/Source/Programs/Enterprise/Datasmith/DatasmithRevitExporter/DatasmithRevit2020/DatasmithRevit2020.csproj
--- a/Engine/Source/Programs/Enterprise/Datasmith/DatasmithRevitExporter/DatasmithRevit2020/DatasmithRevit2020.csproj
+++ b/Engine/Source/Programs/Enterprise/Datasmith/DatasmithRevitExporter/DatasmithRevit2020/DatasmithRevit2020.csproj
@@ -92,7 +92,7 @@
 xcopy "$(ProjectDir)$(TargetName).addin" "%25PUBLISH_DIR%25" /Y
 xcopy "$(ProjectDir)$(OutputPath)*.dll" "%25PUBLISH_DIR%25" /Y
 xcopy "$(ProjectDir)..\Resources\Windows\DatasmithIcon32.png" "%25PUBLISH_DIR%25" /Y
-</PostBuildEvent>
+xcopy "$(ProjectDir)..\Resources\Windows\DatasmithRevitDynamoExport.dyn" "%25PUBLISH_DIR%25" /Y</PostBuildEvent>
   </PropertyGroup>
   <!-- To modify your build process, add your task inside one of the targets below and uncomment it.
        Other similar extension points exist, see Microsoft.Common.targets.
diff --git a/Engine/Source/Programs/Enterprise/Datasmith/DatasmithRevitExporter/Private/DatasmithDynamoNode.cs b/Engine/Source/Programs/Enterprise/Datasmith/DatasmithRevitExporter/Private/DatasmithDynamoNode.cs
new file mode 100644
--- /dev/null
+++ b/Engine/Source/Programs/Enterprise/Datasmith/DatasmithRevitExporter/Private/DatasmithDynamoNode.cs
@@ -0,0 +1,102 @@
+﻿// Copyright 1998-2019 Epic Games, Inc. All Rights Reserved.
+
+using System;
+using System.Collections.Generic;
+using System.IO;
+
+using DatasmithRevitExporter;
+using Autodesk.Revit.DB;
+using System.Text.RegularExpressions;
+
+namespace DSNode
+{
+    public class DatasmithDynamoNode
+    {
+		public static void Export3DViewsToDatasmith(Document InDocument, string InOutputPath, List<int> InViewIds)
+		{
+			if (InDocument == null)
+			{
+				Console.WriteLine("Invalid document.");
+				return;
+			}
+
+			Autodesk.Revit.ApplicationServices.Application App = InDocument.Application;
+
+			if (string.Compare(App.VersionNumber, "2018", StringComparison.Ordinal) == 0 && string.Compare(App.SubVersionNumber, "2018.3", StringComparison.Ordinal) < 0)
+			{
+				string Message = string.Format("The running Revit is not supported.\nYou must use Revit 2018.3 or further updates to export.");
+				Console.WriteLine(Message);
+				return;
+			}
+
+			try
+			{
+				// Validate the provided path
+				InOutputPath = Path.GetFullPath(InOutputPath);
+			}
+			catch (Exception e)
+			{
+				Console.WriteLine(e.Message);
+				return;
+			}
+
+			if (!CustomExporter.IsRenderingSupported())
+			{
+				string Message = "3D view rendering is not supported in the running Revit.";
+				Console.WriteLine(Message);
+				return;
+			}
+
+			string DocumentPath = InDocument.PathName;
+
+			Dictionary<ElementId, string> FilePaths = new Dictionary<ElementId, string>();
+			List<View3D> ExportViews = new List<View3D>();
+			foreach (var Id in InViewIds)
+			{
+				View3D View = InDocument.GetElement(new ElementId(Id)) as View3D;
+				if (View != null && !View.IsTemplate && View.CanBePrinted)
+				{
+					ExportViews.Add(View);
+					// Generate file path for each view.
+					string FileName = Regex.Replace($"{View.Name}.udatasmith", @"\s+", "_");
+					FilePaths.Add(View.Id, Path.Combine(InOutputPath, FileName));
+				}
+			}
+
+			// Retrieve the Unreal Datasmith export options.
+			DatasmithRevitExportOptions ExportOptions = new DatasmithRevitExportOptions(InDocument);
+			// Create a custom export context for command Export to Unreal Datasmith.
+			FDatasmithRevitExportContext ExportContext = new FDatasmithRevitExportContext(App, InDocument, FilePaths, ExportOptions);
+
+			using (CustomExporter Exporter = new CustomExporter(InDocument, ExportContext))
+			{
+				try
+				{
+					// The export process will exclude output of geometric objects such as faces and curves,
+					// but the context needs to receive the calls related to Faces or Curves to gather data.
+					// The context always receive their tessellated geometry in form of polymeshes or lines.
+					Exporter.IncludeGeometricObjects = true;
+
+					// The export process should stop in case an error occurs during any of the exporting methods.
+					Exporter.ShouldStopOnError = true;
+
+					// Initiate the export process for all 3D views.
+					foreach (var View in ExportViews)
+					{
+#if REVIT_API_2020
+						Exporter.Export(View as Autodesk.Revit.DB.View);
+#else
+						Exporter.Export(View);
+#endif
+					}
+				}
+				catch (System.Exception InException)
+				{
+					string Message = string.Format("Cannot export the 3D view:\n\n{0}\n\n{1}", InException.Message, InException.StackTrace);
+					Console.WriteLine(Message);
+					return;
+				}
+			}
+		}
+	}
+}
