Enable Text3D plugin for macOS.

#rnx

diff --git a/Engine/Plugins/Experimental/Text3D/Text3D.uplugin b/Engine/Plugins/Experimental/Text3D/Text3D.uplugin
--- a/Engine/Plugins/Experimental/Text3D/Text3D.uplugin
+++ b/Engine/Plugins/Experimental/Text3D/Text3D.uplugin
@@ -20,7 +20,8 @@
       "WhitelistPlatforms": [
         "Win32",
         "Win64",
-        "Linux"
+        "Linux",
+        "Mac"
       ]
     }
   ]
