Copy CL 11107428 from Dev-Enterprise to fix CIS issue in Text3d plugin
#rnx

diff --git a/Engine/Plugins/Experimental/Text3D/Source/Text3D/Private/MeshCreator.cpp b/Engine/Plugins/Experimental/Text3D/Source/Text3D/Private/MeshCreator.cpp
--- a/Engine/Plugins/Experimental/Text3D/Source/Text3D/Private/MeshCreator.cpp
+++ b/Engine/Plugins/Experimental/Text3D/Source/Text3D/Private/MeshCreator.cpp
@@ -62,10 +62,10 @@ void FMeshCreator::SetFrontAndBevelTextureCoordinates(const float Bevel)
 
 	auto SetTextureCoordinates = [Groups, VertexPositions, VertexInstanceUVs, &Box](const EText3DGroupType Type)
 	{
-		const int32 FirstVertex = Groups[static_cast<int32>(Type)].FirstVertex;
-		const int32 LastVertex = Groups[static_cast<int32>(Type) + 1].FirstVertex;
+		const int32 TypeFirstVertex = Groups[static_cast<int32>(Type)].FirstVertex;
+		const int32 TypeLastVertex = Groups[static_cast<int32>(Type) + 1].FirstVertex;
 
-		for (int32 Index = FirstVertex; Index < LastVertex; Index++)
+		for (int32 Index = TypeFirstVertex; Index < TypeLastVertex; Index++)
 		{
 			const FVector Position = VertexPositions[FVertexID(Index)];
 			const FVector2D TextureCoordinate = (FVector2D(Position.Y, Position.Z) - Box.Min) / Box.Max;
