copy fix for Text3D cis error from Dev-Editor
#rnx

diff --git a/Engine/Plugins/Experimental/Text3D/Source/Text3D/Private/MeshCreator.cpp b/Engine/Plugins/Experimental/Text3D/Source/Text3D/Private/MeshCreator.cpp
--- a/Engine/Plugins/Experimental/Text3D/Source/Text3D/Private/MeshCreator.cpp
+++ b/Engine/Plugins/Experimental/Text3D/Source/Text3D/Private/MeshCreator.cpp
@@ -385,9 +385,9 @@ void FMeshCreator::MirrorGroup(const EText3DGroupType TypeIn, const EText3DGroup
 	TVertexInstanceAttributesRef<FVector> VertexTangents = StaticMeshAttributes.GetVertexInstanceTangents();
 	TVertexInstanceAttributesRef<FVector2D> VertexUVs = StaticMeshAttributes.GetVertexInstanceUVs();
 
-	for (int32 Index = 0; Index < VerticesInNum; Index++)
+	for (int32 VertexIndex = 0; VertexIndex < VerticesInNum; VertexIndex++)
 	{
-		const FVertexID VertexID(GroupIn.FirstVertex + Index);
+		const FVertexID VertexID(GroupIn.FirstVertex + VertexIndex);
 		const FVertexInstanceID InstanceID(static_cast<uint32>(VertexID.GetValue()));
 
 		const FVector Position = VertexPositions[VertexID];
@@ -399,9 +399,9 @@ void FMeshCreator::MirrorGroup(const EText3DGroupType TypeIn, const EText3DGroup
 
 	Data->AddTriangles(TrianglesInNum);
 
-	for (int32 Index = 0; Index < TrianglesInNum; Index++)
+	for (int32 TriangleIndex = 0; TriangleIndex < TrianglesInNum; TriangleIndex++)
 	{
-		const FMeshTriangle& Triangle = MeshDescription.Triangles()[FTriangleID(GroupIn.FirstTriangle + Index)];
+		const FMeshTriangle& Triangle = MeshDescription.Triangles()[FTriangleID(GroupIn.FirstTriangle + TriangleIndex)];
 
 		auto InstanceID = [TotalVerticesNum, &Triangle, GroupIn](const int32 Index)
 		{
