Blacklist Servers for Text3D plugin

diff --git a/Engine/Plugins/Experimental/Text3D/Text3D.uplugin b/Engine/Plugins/Experimental/Text3D/Text3D.uplugin
--- a/Engine/Plugins/Experimental/Text3D/Text3D.uplugin
+++ b/Engine/Plugins/Experimental/Text3D/Text3D.uplugin
@@ -17,12 +17,14 @@
       "Name": "Text3D",
       "Type": "Runtime",
       "LoadingPhase": "Default"
+	  "BlacklistTargets" : [ "Server" ]
     }
   ],
   "Plugins": [
     {
       "Name": "GeometryProcessing",
       "Enabled": true
+	  "BlacklistTargets" : [ "Server" ]
     }
   ]
 }
