Add missing comma to Text3D.uplugin

diff --git a/Engine/Plugins/Experimental/Text3D/Text3D.uplugin b/Engine/Plugins/Experimental/Text3D/Text3D.uplugin
--- a/Engine/Plugins/Experimental/Text3D/Text3D.uplugin
+++ b/Engine/Plugins/Experimental/Text3D/Text3D.uplugin
@@ -16,15 +16,15 @@
     {
       "Name": "Text3D",
       "Type": "Runtime",
-      "LoadingPhase": "Default"
-	  "BlacklistTargets" : [ "Server" ]
+      "LoadingPhase": "Default",
+	  "BlacklistTargets": [ "Server" ]
     }
   ],
   "Plugins": [
     {
       "Name": "GeometryProcessing",
-      "Enabled": true
-	  "BlacklistTargets" : [ "Server" ]
+      "Enabled": true,
+	  "BlacklistTargets": [ "Server" ]
     }
   ]
 }
