Text 3D fix for crash on duplicate

diff --git a/Engine/Plugins/Experimental/Text3D/Source/Text3D/Private/Text3DComponent.cpp b/Engine/Plugins/Experimental/Text3D/Source/Text3D/Private/Text3DComponent.cpp
--- a/Engine/Plugins/Experimental/Text3D/Source/Text3D/Private/Text3DComponent.cpp
+++ b/Engine/Plugins/Experimental/Text3D/Source/Text3D/Private/Text3DComponent.cpp
@@ -1,21 +1,18 @@
 // Copyright Epic Games, Inc. All Rights Reserved.
 
 #include "Text3DComponent.h"
+#include "Glyph.h"
+#include "MeshCreator.h"
 #include "Text3DPrivate.h"
 #include "TextShaper.h"
-#include "MeshCreator.h"
-#include "Glyph.h"
 
-#include "ContourList.h"
-#include "Data.h"
 #include "Text3DEngineSubsystem.h"
 #include "Components/StaticMeshComponent.h"
-#include "Engine/Font.h"
 #include "Engine/Engine.h"
+#include "Engine/Font.h"
 #include "Materials/Material.h"
 #include "UObject/ConstructorHelpers.h"
 
-
 #define LOCTEXT_NAMESPACE "Text3D"
 
 struct FText3DShapedText
@@ -267,9 +264,12 @@ UMaterialInterface* UText3DComponent::GetMaterial(const EText3DGroupType Type) c
 	{
 		return BackMaterial;
 	}
+		
+	default:
+	{
+		return nullptr;
+	}
 	}
-
-	return nullptr;
 }
 
 void UText3DComponent::SetMaterial(const EText3DGroupType Type, UMaterialInterface* Value)
@@ -302,6 +302,11 @@ void UText3DComponent::SetMaterial(const EText3DGroupType Type, UMaterialInterfa
 			BevelMaterial = Value;
 			break;
 		}
+
+		default:
+		{
+			return;
+		}
 		}
 
 		UpdateMaterial(Type, Value);
@@ -549,7 +554,7 @@ FVector UText3DComponent::GetLineLocation(int32 LineIndex)
 void UText3DComponent::UpdateTransforms()
 {
 	CalculateTextWidth();
-	FVector Scale = GetTextScale();
+	const FVector Scale = GetTextScale();
 	TextRoot->SetRelativeScale3D(Scale);
 
 	int32 GlyphIndex = 0;
@@ -560,7 +565,7 @@ void UText3DComponent::UpdateTransforms()
 
 		for (int32 LineGlyph = 0; LineGlyph < Line.GlyphsToRender.Num(); LineGlyph++)
 		{
-			FVector CharLocation = Location;
+			const FVector CharLocation = Location;
 			Location.Y += Line.GetAdvanced(LineGlyph, Kerning, WordSpacing);
 			if (!Line.GlyphsToRender[LineGlyph].bIsVisible)
 			{
@@ -575,8 +580,6 @@ void UText3DComponent::UpdateTransforms()
 
 			GlyphIndex++;
 		}
-
-		Location.Z -= ShapedText->LineHeight + LineSpacing;
 	}
 }
 
@@ -619,7 +622,7 @@ void UText3DComponent::BuildTextMesh(const bool bCleanCache)
 
 	UText3DEngineSubsystem* Subsystem = GEngine->GetEngineSubsystem<UText3DEngineSubsystem>();
 	FCachedFontData& CachedFontData = Subsystem->GetCachedFontData(Font);
-	FT_Face Face = CachedFontData.GetFreeTypeFace();
+	const FT_Face Face = CachedFontData.GetFreeTypeFace();
 	if (!Face)
 	{
 		UE_LOG(LogText3D, Error, TEXT("Failed to load font data '%s'"), *CachedFontData.GetFontName());
@@ -638,7 +641,6 @@ void UText3DComponent::BuildTextMesh(const bool bCleanCache)
 	CalculateTextWidth();
 	TextRoot->SetRelativeScale3D(GetTextScale());
 
-
 	int32 GlyphIndex = 0;
 	for (int32 LineIndex = 0; LineIndex < ShapedText->Lines.Num(); LineIndex++)
 	{
@@ -662,20 +664,20 @@ void UText3DComponent::BuildTextMesh(const bool bCleanCache)
 				continue;
 			}
 
-			int32 GlyphId = GlyphIndex++;
-
-			FString CharachterKerningName = FString::Printf(TEXT("CharachterKerning%d"), GlyphId);
-			USceneComponent* CharachterKerningComponent = NewObject<USceneComponent>(this, FName(*CharachterKerningName));
+			const int32 GlyphId = GlyphIndex++;
 
+			const FName CharacterKerningComponentName = MakeUniqueObjectName(this, USceneComponent::StaticClass(), FName(*FString::Printf(TEXT("CharacterKerning%d"), GlyphId)));
+			USceneComponent* CharacterKerningComponent = NewObject<USceneComponent>(this, CharacterKerningComponentName);
+			
 #if WITH_EDITOR
-			CharachterKerningComponent->SetIsVisualizationComponent(true);
+			CharacterKerningComponent->SetIsVisualizationComponent(true);
 #endif
-			CharachterKerningComponent->AttachToComponent(TextRoot, FAttachmentTransformRules::KeepRelativeTransform);
-			CharachterKerningComponent->RegisterComponent();
-			CharacterKernings.Add(CharachterKerningComponent);
+			CharacterKerningComponent->AttachToComponent(TextRoot, FAttachmentTransformRules::KeepRelativeTransform);
+			CharacterKerningComponent->RegisterComponent();
+			CharacterKernings.Add(CharacterKerningComponent);
 
-			FString StatichMeshComponentName = FString::Printf(TEXT("StaticMeshComponent%d"), GlyphId);
-			UStaticMeshComponent* StaticMeshComponent = NewObject<UStaticMeshComponent>(this, FName(*StatichMeshComponentName));
+			const FName StaticMeshComponentName = MakeUniqueObjectName(this, UStaticMeshComponent::StaticClass(), FName(*FString::Printf(TEXT("StaticMeshComponent%d"), GlyphId)));
+			UStaticMeshComponent* StaticMeshComponent = NewObject<UStaticMeshComponent>(this, StaticMeshComponentName);
 #if WITH_EDITOR
 			StaticMeshComponent->SetIsVisualizationComponent(true);
 #endif
@@ -686,19 +688,17 @@ void UText3DComponent::BuildTextMesh(const bool bCleanCache)
 			CharacterMeshes.Add(StaticMeshComponent);
 
 			GetOwner()->AddInstanceComponent(StaticMeshComponent);
-			StaticMeshComponent->AttachToComponent(CharachterKerningComponent, FAttachmentTransformRules::KeepRelativeTransform);
+			StaticMeshComponent->AttachToComponent(CharacterKerningComponent, FAttachmentTransformRules::KeepRelativeTransform);
 
 			FTransform Transform;
 			Transform.SetLocation(GlyphLocation);
-			CharachterKerningComponent->SetRelativeTransform(Transform);
+			CharacterKerningComponent->SetRelativeTransform(Transform);
 		}
-
-		Location.Z -= ShapedText->LineHeight + LineSpacing;
 	}
-
+	
 	for (int32 Index = 0; Index < static_cast<int32>(EText3DGroupType::TypeCount); Index++)
 	{
-		EText3DGroupType Type = static_cast<EText3DGroupType>(Index);
+		const EText3DGroupType Type = static_cast<EText3DGroupType>(Index);
 		UpdateMaterial(Type, GetMaterial(Type));
 	}
 
@@ -726,7 +726,7 @@ float UText3DComponent::MaxBevel() const
 
 void UText3DComponent::UpdateMaterial(const EText3DGroupType Type, UMaterialInterface* Material)
 {
-	int32 Index = static_cast<int32>(Type);
+	const int32 Index = static_cast<int32>(Type);
 	for (UStaticMeshComponent* StaticMeshComponent : CharacterMeshes)
 	{
 		StaticMeshComponent->SetMaterial(Index, Material);
@@ -736,7 +736,7 @@ void UText3DComponent::UpdateMaterial(const EText3DGroupType Type, UMaterialInte
 void UText3DComponent::OnVisibilityChanged()
 {
 	Super::OnVisibilityChanged();
-	bool Visibility = GetVisibleFlag();
+	const bool Visibility = GetVisibleFlag();
 	for (UStaticMeshComponent* StaticMeshComponent : CharacterMeshes)
 	{
 		StaticMeshComponent->SetVisibility(Visibility);
@@ -756,7 +756,7 @@ void UText3DComponent::GetBounds(FVector& Origin, FVector& BoxExtent)
 {
 	FBox Box(ForceInit);
 
-	for (UStaticMeshComponent* StaticMeshComponent : CharacterMeshes)
+	for (const UStaticMeshComponent* StaticMeshComponent : CharacterMeshes)
 	{
 		Box += StaticMeshComponent->Bounds.GetBox();
 	}
