Text3D: Fix color groups not applying correct group

diff --git a/Engine/Plugins/Experimental/Text3D/Source/Text3D/Private/Glyph.cpp b/Engine/Plugins/Experimental/Text3D/Source/Text3D/Private/Glyph.cpp
--- a/Engine/Plugins/Experimental/Text3D/Source/Text3D/Private/Glyph.cpp
+++ b/Engine/Plugins/Experimental/Text3D/Source/Text3D/Private/Glyph.cpp
@@ -32,7 +32,9 @@ void FText3DGlyph::Build(UStaticMesh* StaticMesh, UMaterial* DefaultMaterial)
 
 	UStaticMesh::FBuildMeshDescriptionsParams Params;
 	Params.bCommitMeshDescription = true;
+#if !WITH_EDITOR
 	Params.bFastBuild = true;
+#endif
 	StaticMesh->BuildFromMeshDescriptions(MeshDescriptions, Params);
 }
 
