Text3D casts shadows and allows outline offset

diff --git a/Engine/Plugins/Experimental/Text3D/Source/Text3D/Private/Data.cpp b/Engine/Plugins/Experimental/Text3D/Source/Text3D/Private/Data.cpp
--- a/Engine/Plugins/Experimental/Text3D/Source/Text3D/Private/Data.cpp
+++ b/Engine/Plugins/Experimental/Text3D/Source/Text3D/Private/Data.cpp
@@ -166,17 +166,17 @@ FVector2D FData::Expanded(const FPartConstPtr& Point) const
 	return Point->Expanded(ExpandTarget - Point->DoneExpand);
 }
 
-void FData::FillEdge(const FPartPtr& Edge, const bool bSkipLastTriangle)
+void FData::FillEdge(const FPartPtr& Edge, const bool bSkipLastTriangle, bool bFlipNormals)
 {
 	const FPartPtr EdgeA = Edge;
 	const FPartPtr EdgeB = Edge->Next;
 
-	MakeTriangleFanAlongNormal(EdgeB, EdgeA, false, true);
-	MakeTriangleFanAlongNormal(EdgeA, EdgeB, true, false);
+	MakeTriangleFanAlongNormal(EdgeB, EdgeA, bFlipNormals, true);
+	MakeTriangleFanAlongNormal(EdgeA, EdgeB, !bFlipNormals, false);
 
 	if (!bSkipLastTriangle)
 	{
-		MakeTriangleFanAlongNormal(EdgeB, EdgeA, false, false);
+		MakeTriangleFanAlongNormal(EdgeB, EdgeA, bFlipNormals, false);
 	}
 	else
 	{
@@ -208,9 +208,8 @@ void FData::MakeTriangleFanAlongNormal(const FPartConstPtr& Cap, const FPartPtr&
 	for (int32 Index = 0; Index < Count; Index++)
 	{
 		AddTriangle((bNormalIsCapNext ? Cap->PathNext : Cap->PathPrev)[0],
-				Path[bNormalIsCapNext ? Index + 1 : Index],
-				Path[bNormalIsCapNext ? Index : Index + 1]
-			);
+						Path[bNormalIsCapNext ? Index + 1 : Index],
+						Path[bNormalIsCapNext ? Index : Index + 1]);
 	}
 
 	// Remove covered vertices from path
diff --git a/Engine/Plugins/Experimental/Text3D/Source/Text3D/Private/Data.h b/Engine/Plugins/Experimental/Text3D/Source/Text3D/Private/Data.h
--- a/Engine/Plugins/Experimental/Text3D/Source/Text3D/Private/Data.h
+++ b/Engine/Plugins/Experimental/Text3D/Source/Text3D/Private/Data.h
@@ -69,8 +69,9 @@ public:
 	 * Make triangulation of edge along paths of it's vertices (from end of previous triangulation to result of points' expansion). Removes covered points' indices from paths.
 	 * @param Edge - Edge that has to be filled.
 	 * @param bSkipLastTriangle - Do not create last triangle (furthest from end of previous triangulation).
+	 * @param bFlipNormals - Reverse the geometry facing direction.
 	 */
-	void FillEdge(const FPartPtr& Edge, const bool bSkipLastTriangle);
+	void FillEdge(const FPartPtr& Edge, const bool bSkipLastTriangle, bool bFlipNormals = false);
 
 private:
 	TSharedPtr<FText3DGlyph> Glyph;
diff --git a/Engine/Plugins/Experimental/Text3D/Source/Text3D/Private/MeshCreator.cpp b/Engine/Plugins/Experimental/Text3D/Source/Text3D/Private/MeshCreator.cpp
--- a/Engine/Plugins/Experimental/Text3D/Source/Text3D/Private/MeshCreator.cpp
+++ b/Engine/Plugins/Experimental/Text3D/Source/Text3D/Private/MeshCreator.cpp
@@ -11,25 +11,24 @@
 
 using namespace UE::Geometry;
 
-constexpr float FMeshCreator::OutlineExpand;
-
 FMeshCreator::FMeshCreator() :
 	Glyph(MakeShared<FText3DGlyph>()),
 	Data(MakeShared<FData>(Glyph))
 {
 }
 
-void FMeshCreator::CreateMeshes(const TSharedContourNode& Root, const bool bOutline, const float Extrude, const float Bevel, const EText3DBevelType Type, const int32 BevelSegments)
+void FMeshCreator::CreateMeshes(const TSharedContourNode& Root, const float Extrude, const float Bevel, const EText3DBevelType Type, const int32 BevelSegments, const bool bOutline, const float OutlineExpand)
 {
-	CreateFrontMesh(Root, bOutline);
+	CreateFrontMesh(Root, bOutline, OutlineExpand);
 	if (Contours->Num() == 0)
 	{
 		return;
 	}
 
+	const bool bFlipNormals = FMath::Sign(Data->GetPlannedExpand()) < 0;
 	const float BevelLocal = bOutline ? 0.f : Bevel;
 	CreateBevelMesh(BevelLocal, Type, BevelSegments);
-	CreateExtrudeMesh(Extrude, BevelLocal, Type);
+	CreateExtrudeMesh(Extrude, BevelLocal, Type, bFlipNormals);
 }
 
 void FMeshCreator::SetFrontAndBevelTextureCoordinates(const float Bevel)
@@ -94,7 +93,7 @@ void FMeshCreator::BuildMesh(UStaticMesh* StaticMesh, class UMaterial* DefaultMa
 	Glyph->Build(StaticMesh, DefaultMaterial);
 }
 
-void FMeshCreator::CreateFrontMesh(const TSharedContourNode& Root, const bool bOutline)
+void FMeshCreator::CreateFrontMesh(const TSharedContourNode& Root, const bool bOutline, const float& OutlineExpand)
 {
 	int32 VertexCount = 0;
 	AddToVertexCount(Root, VertexCount);
@@ -110,7 +109,7 @@ void FMeshCreator::CreateFrontMesh(const TSharedContourNode& Root, const bool bO
 
 	if (bOutline)
 	{
-		MakeOutline();
+		MakeOutline(OutlineExpand);
 	}
 }
 
@@ -176,7 +175,7 @@ void FMeshCreator::CreateBevelMesh(const float Bevel, const EText3DBevelType Typ
 	}
 }
 
-void FMeshCreator::CreateExtrudeMesh(float Extrude, float Bevel, const EText3DBevelType Type)
+void FMeshCreator::CreateExtrudeMesh(float Extrude, float Bevel, const EText3DBevelType Type, bool bFlipNormals)
 {
 	if (Type != EText3DBevelType::HalfCircle)
 	{
@@ -185,10 +184,9 @@ void FMeshCreator::CreateExtrudeMesh(float Extrude, float Bevel, const EText3DBe
 
 	if (Type != EText3DBevelType::HalfCircle && Type != EText3DBevelType::Engraved)
 	{
-	Extrude -= Bevel * 2.0f;
+		Extrude -= Bevel * 2.0f;
 	}
 
-
 	Data->SetCurrentGroup(EText3DGroupType::Extrude, 0.f);
 
 	const FVector2D Normal(1.f, 0.f);
@@ -291,7 +289,7 @@ void FMeshCreator::CreateExtrudeMesh(float Extrude, float Bevel, const EText3DBe
 
 		for (const FPartPtr& Edge : Contour)
 		{
-			Data->FillEdge(Edge, false);
+			Data->FillEdge(Edge, false, bFlipNormals);
 		}
 	}
 }
@@ -426,7 +424,7 @@ void FMeshCreator::TriangulateAndConvert(const TSharedContourNode& Node, int32&
 	}
 }
 
-void FMeshCreator::MakeOutline()
+void FMeshCreator::MakeOutline(float OutlineExpand)
 {
 	FContourList InitialContours = *Contours;
 
@@ -588,6 +586,7 @@ void FMeshCreator::BevelPartsWithoutIntersectingNormals()
 	Data->SetTarget(Data->GetPlannedExtrude(), Data->GetPlannedExpand());
 	const float MaxExpand = Data->GetPlannedExpand();
 
+	const bool bFlipNormals = FMath::Sign(Data->GetPlannedExpand()) < 0;
 	for (FContour& Contour : *Contours)
 	{
 		for (const FPartPtr& Point : Contour)
@@ -605,7 +604,7 @@ void FMeshCreator::BevelPartsWithoutIntersectingNormals()
 
 		for (const FPartPtr& Edge : Contour)
 		{
-			Data->FillEdge(Edge, false);
+			Data->FillEdge(Edge, false, bFlipNormals);
 		}
 	}
 }
diff --git a/Engine/Plugins/Experimental/Text3D/Source/Text3D/Private/MeshCreator.h b/Engine/Plugins/Experimental/Text3D/Source/Text3D/Private/MeshCreator.h
--- a/Engine/Plugins/Experimental/Text3D/Source/Text3D/Private/MeshCreator.h
+++ b/Engine/Plugins/Experimental/Text3D/Source/Text3D/Private/MeshCreator.h
@@ -16,21 +16,20 @@ public:
 
 	/**
 	 * CreateMeshes
-	 * @param Root - Tree of contours.
-	 * @param bOutline - Front face has outline and is not filled.
+	 * @param Root - Tree of contours.	 
 	 * @param Extrude - Orthogonal (to front cap) offset value.
 	 * @param Bevel - Bevel value (bevel happens before extrude).
 	 * @param Type - Defines shape of beveled part.
 	 * @param BevelSegments - Segments count.
+	 * @param bOutline - Front face has outline and is not filled.
+	 * @param OutlineExpand - Offsets the outline by the specified amount.
 	 */
-	void CreateMeshes(const TSharedContourNode& Root, const bool bOutline, const float Extrude, const float Bevel, const EText3DBevelType Type, const int32 BevelSegments);
+	void CreateMeshes(const TSharedContourNode& Root, const float Extrude, const float Bevel, const EText3DBevelType Type, const int32 BevelSegments, const bool bOutline, const float OutlineExpand = 0.5f);
 	void SetFrontAndBevelTextureCoordinates(const float Bevel);
 	void MirrorGroups(const float Extrude);
 	void BuildMesh(UStaticMesh* StaticMesh, class UMaterial* DefaultMaterial);
 
 private:
-	constexpr static float OutlineExpand = 0.5f;
-
 	TSharedPtr<class FText3DGlyph> Glyph;
 	TSharedRef<class FData> Data;
 	TSharedPtr<class FContourList> Contours;
@@ -40,8 +39,9 @@ private:
 	 * Create 'Front' part of glyph.
 	 * @param Root - Tree root.
 	 * @param bOutline - Front face has outline and is not filled.
+	 * @param OutlineExpand - Offsets the outline by the specified amount.
 	 */
-	void CreateFrontMesh(const TSharedContourNode& Root, const bool bOutline);
+	void CreateFrontMesh(const TSharedContourNode& Root, const bool bOutline, const float& OutlineExpand = 0.5f);
 	/**
 	 * Create 'Bevel' part of glyph (actually half of it, will be mirrored later).
 	 * @param Bevel - Bevel value (bevel happens before extrude).
@@ -54,8 +54,9 @@ private:
 	 * @param Extrude - Extrude value.
 	 * @param Bevel - Bevel value (bevel happens before extrude).
 	 * @param Type - Defines shape of beveled part.
+	 * @param bFlipNormals - Reverse the geometry facing direction.
 	 */
-	void CreateExtrudeMesh(float Extrude, float Bevel, const EText3DBevelType Type);
+	void CreateExtrudeMesh(float Extrude, float Bevel, const EText3DBevelType Type, bool bFlipNormals = false);
 
 	void MirrorGroup(const EText3DGroupType TypeIn, const EText3DGroupType TypeOut, const float Extrude);
 
@@ -74,7 +75,7 @@ private:
 	 */
 	void TriangulateAndConvert(const TSharedContourNode& Node, int32& OutVertexIndex, const bool bOutline);
 
-	void MakeOutline();
+	void MakeOutline(float OutlineExpand = 0.5f);
 
 	void BevelLinearWithSegments(const float Extrude, const float Expand, const int32 BevelSegments, const FVector2D Normal);
 	void BevelCurve(const float Angle, const int32 BevelSegments, TFunction<FVector2D(const float CurrentCos, const float CurrentSin, const float NextCos, const float Next)> ComputeOffset);
diff --git a/Engine/Plugins/Experimental/Text3D/Source/Text3D/Private/Text3DComponent.cpp b/Engine/Plugins/Experimental/Text3D/Source/Text3D/Private/Text3DComponent.cpp
--- a/Engine/Plugins/Experimental/Text3D/Source/Text3D/Private/Text3DComponent.cpp
+++ b/Engine/Plugins/Experimental/Text3D/Source/Text3D/Private/Text3DComponent.cpp
@@ -71,6 +71,7 @@ UText3DComponent::UText3DComponent() :
 
 	Text = LOCTEXT("DefaultText", "Text");
 	bOutline = false;
+	OutlineExpand = 0.5f;
 	Extrude = 5.0f;
 	Bevel = 0.0f;
 	BevelType = EText3DBevelType::Convex;
@@ -169,7 +170,17 @@ void UText3DComponent::SetOutline(const bool bValue)
 {
 	if (bOutline != bValue)
 	{
-		bOutline = bValue;
+		bOutline = bValue;		
+		Rebuild();
+	}
+}
+
+void UText3DComponent::SetOutlineExpand(const float Value)
+{
+	const float NewValue = Value;
+	if (!FMath::IsNearlyEqual(OutlineExpand, NewValue))
+	{
+		OutlineExpand = NewValue;
 		Rebuild();
 	}
 }
@@ -418,6 +429,21 @@ void UText3DComponent::SetFreeze(const bool bFreeze)
 	}
 }
 
+void UText3DComponent::SetCastShadow(bool NewCastShadow)
+{
+	if(NewCastShadow != bCastShadow)
+	{
+		bCastShadow = NewCastShadow;
+
+		for (UStaticMeshComponent* MeshComponent : CharacterMeshes)
+		{
+			MeshComponent->SetCastShadow(bCastShadow);
+		}
+		
+		MarkRenderStateDirty();
+	}
+}
+
 int32 UText3DComponent::GetGlyphCount()
 {
 	return TextRoot->GetNumChildrenComponents();
@@ -619,7 +645,7 @@ void UText3DComponent::BuildTextMesh(const bool bCleanCache)
 	{
 		return;
 	}
-
+	
 	UText3DEngineSubsystem* Subsystem = GEngine->GetEngineSubsystem<UText3DEngineSubsystem>();
 	FCachedFontData& CachedFontData = Subsystem->GetCachedFontData(Font);
 	const FT_Face Face = CachedFontData.GetFreeTypeFace();
@@ -630,7 +656,7 @@ void UText3DComponent::BuildTextMesh(const bool bCleanCache)
 	}
 
 	CachedCounterReferences.Add(CachedFontData.GetCacheCounter());
-	CachedCounterReferences.Add(CachedFontData.GetMeshesCacheCounter(bOutline, Extrude, Bevel, BevelType, BevelSegments));
+	CachedCounterReferences.Add(CachedFontData.GetMeshesCacheCounter(FGlyphMeshParameters{Extrude, Bevel, BevelType, BevelSegments, bOutline, OutlineExpand}));
 
 	ShapedText->Reset();
 	ShapedText->LineHeight = Face->size->metrics.height * FontInverseScale;
@@ -658,7 +684,7 @@ void UText3DComponent::BuildTextMesh(const bool bCleanCache)
 				continue;
 			}
 
-			UStaticMesh* CachedMesh = CachedFontData.GetGlyphMesh(ShapedGlyph.GlyphIndex, bOutline, Extrude, Bevel, BevelType, BevelSegments);
+			UStaticMesh* CachedMesh = CachedFontData.GetGlyphMesh(ShapedGlyph.GlyphIndex, FGlyphMeshParameters{Extrude, Bevel, BevelType, BevelSegments, bOutline, OutlineExpand});
 			if (!CachedMesh)
 			{
 				continue;
@@ -685,6 +711,7 @@ void UText3DComponent::BuildTextMesh(const bool bCleanCache)
 			StaticMeshComponent->RegisterComponent();
 			StaticMeshComponent->SetVisibility(GetVisibleFlag());
 			StaticMeshComponent->SetHiddenInGame(bHiddenInGame);
+			StaticMeshComponent->SetCastShadow(bCastShadow);
 			CharacterMeshes.Add(StaticMeshComponent);
 
 			GetOwner()->AddInstanceComponent(StaticMeshComponent);
diff --git a/Engine/Plugins/Experimental/Text3D/Source/Text3D/Private/Text3DEngineSubsystem.cpp b/Engine/Plugins/Experimental/Text3D/Source/Text3D/Private/Text3DEngineSubsystem.cpp
--- a/Engine/Plugins/Experimental/Text3D/Source/Text3D/Private/Text3DEngineSubsystem.cpp
+++ b/Engine/Plugins/Experimental/Text3D/Source/Text3D/Private/Text3DEngineSubsystem.cpp
@@ -187,27 +187,17 @@ TSharedPtr<int32> FCachedFontData::GetCacheCounter()
 	return CacheCounter;
 }
 
-TSharedPtr<int32> FCachedFontData::GetMeshesCacheCounter(bool bOutline, float Extrude, float Bevel, EText3DBevelType BevelType, float BevelSegments)
+TSharedPtr<int32> FCachedFontData::GetMeshesCacheCounter(const FGlyphMeshParameters& Parameters)
 {
-	uint32 HashParameters = 0;
-	HashParameters = HashCombine(HashParameters, GetTypeHash(bOutline));
-	HashParameters = HashCombine(HashParameters, GetTypeHash(Extrude));
-	HashParameters = HashCombine(HashParameters, GetTypeHash(Bevel));
-	HashParameters = HashCombine(HashParameters, GetTypeHash(BevelType));
-	HashParameters = HashCombine(HashParameters, GetTypeHash(BevelSegments));
+	const uint32 HashParameters = GetTypeHash(Parameters);
 	FCachedFontMeshes& CachedMeshes = Meshes.FindOrAdd(HashParameters);
 
 	return CachedMeshes.GetCacheCounter();
 }
 
-UStaticMesh* FCachedFontData::GetGlyphMesh(uint32 GlyphIndex, bool bOutline, float Extrude, float Bevel, EText3DBevelType BevelType, float BevelSegments)
+UStaticMesh* FCachedFontData::GetGlyphMesh(uint32 GlyphIndex, const FGlyphMeshParameters& Parameters)
 {
-	uint32 HashParameters = 0;
-	HashParameters = HashCombine(HashParameters, GetTypeHash(bOutline));
-	HashParameters = HashCombine(HashParameters, GetTypeHash(Extrude));
-	HashParameters = HashCombine(HashParameters, GetTypeHash(Bevel));
-	HashParameters = HashCombine(HashParameters, GetTypeHash(BevelType));
-	HashParameters = HashCombine(HashParameters, GetTypeHash(BevelSegments));
+	const uint32 HashParameters = GetTypeHash(Parameters);
 	FCachedFontMeshes& CachedMeshes = Meshes.FindOrAdd(HashParameters);
 
 	UStaticMesh** CachedStaticMesh = CachedMeshes.Glyphs.Find(GlyphIndex);
@@ -220,8 +210,7 @@ UStaticMesh* FCachedFontData::GetGlyphMesh(uint32 GlyphIndex, bool bOutline, flo
 	HashGroup = HashCombine(HashGroup, GetTypeHash(Font));
 	HashGroup = HashCombine(HashGroup, GetTypeHash(GlyphIndex));
 	FString StaticMeshName = FString::Printf(TEXT("Text3D_Char_%u_%u"), HashGroup, HashParameters);
-
-
+	
 	FMeshCreator MeshCreator;
 	TSharedContourNode Root = GetGlyphContours(GlyphIndex);
 	if (Root->Children.Num() == 0)
@@ -229,9 +218,9 @@ UStaticMesh* FCachedFontData::GetGlyphMesh(uint32 GlyphIndex, bool bOutline, flo
 		return nullptr;
 	}
 
-	MeshCreator.CreateMeshes(Root, bOutline, Extrude, Bevel, BevelType, BevelSegments);
-	MeshCreator.SetFrontAndBevelTextureCoordinates(Bevel);
-	MeshCreator.MirrorGroups(Extrude);
+	MeshCreator.CreateMeshes(Root, Parameters.Extrude, Parameters.Bevel, Parameters.BevelType, Parameters.BevelSegments, Parameters.bOutline, Parameters.OutlineExpand);
+	MeshCreator.SetFrontAndBevelTextureCoordinates(Parameters.Bevel);
+	MeshCreator.MirrorGroups(Parameters.Extrude);
 
 	UText3DEngineSubsystem* Subsystem = GEngine->GetEngineSubsystem<UText3DEngineSubsystem>();
 	UStaticMesh* StaticMesh = NewObject<UStaticMesh>(Subsystem, *StaticMeshName);
diff --git a/Engine/Plugins/Experimental/Text3D/Source/Text3D/Private/Text3DEngineSubsystem.h b/Engine/Plugins/Experimental/Text3D/Source/Text3D/Private/Text3DEngineSubsystem.h
--- a/Engine/Plugins/Experimental/Text3D/Source/Text3D/Private/Text3DEngineSubsystem.h
+++ b/Engine/Plugins/Experimental/Text3D/Source/Text3D/Private/Text3DEngineSubsystem.h
@@ -1,6 +1,7 @@
 // Copyright Epic Games, Inc. All Rights Reserved.
 
 #pragma once
+
 #include "CoreMinimal.h"
 #include "Subsystems/EngineSubsystem.h"
 #include "Text3DPrivate.h"
@@ -10,6 +11,42 @@
 
 #include "Text3DEngineSubsystem.generated.h"
 
+USTRUCT()
+struct FGlyphMeshParameters
+{
+	GENERATED_BODY()
+
+public:
+	UPROPERTY()
+	float Extrude = 5.0f;
+
+	UPROPERTY()
+	float Bevel = 0.0f;
+
+	UPROPERTY()
+	EText3DBevelType BevelType = EText3DBevelType::Convex;
+
+	UPROPERTY()
+	int32 BevelSegments = 8;
+
+	UPROPERTY()
+	bool bOutline = false;
+
+	UPROPERTY()
+	float OutlineExpand = 0.5f;
+};
+
+inline uint32 GetTypeHash(const FGlyphMeshParameters A)
+{
+	uint32 HashParameters = 0;
+	HashParameters = HashCombine(HashParameters, GetTypeHash(A.Extrude));
+	HashParameters = HashCombine(HashParameters, GetTypeHash(A.Bevel));
+    HashParameters = HashCombine(HashParameters, GetTypeHash(A.BevelType));
+	HashParameters = HashCombine(HashParameters, GetTypeHash(A.BevelSegments));
+	HashParameters = HashCombine(HashParameters, GetTypeHash(A.bOutline));
+	HashParameters = HashCombine(HashParameters, GetTypeHash(A.OutlineExpand));
+	return HashParameters;
+}
 
 USTRUCT()
 struct FCachedFontMeshes
@@ -48,9 +85,9 @@ public:
 
 	uint32 GetTypefaceFontDataHash();
 	TSharedPtr<int32> GetCacheCounter();
-	TSharedPtr<int32> GetMeshesCacheCounter(bool bOutline, float Extrude, float Bevel, EText3DBevelType BevelType, float BevelSegments);
+	TSharedPtr<int32> GetMeshesCacheCounter(const FGlyphMeshParameters& Parameters);
 
-	UStaticMesh* GetGlyphMesh(uint32 GlyphIndex, bool bOutline, float Extrude, float Bevel, EText3DBevelType BevelType, float BevelSegments);
+	UStaticMesh* GetGlyphMesh(uint32 GlyphIndex, const FGlyphMeshParameters& Parameters);
 	TSharedContourNode GetGlyphContours(uint32 GlyphIndex);
 
 	UPROPERTY()
diff --git a/Engine/Plugins/Experimental/Text3D/Source/Text3D/Public/Text3DComponent.h b/Engine/Plugins/Experimental/Text3D/Source/Text3D/Public/Text3DComponent.h
--- a/Engine/Plugins/Experimental/Text3D/Source/Text3D/Public/Text3DComponent.h
+++ b/Engine/Plugins/Experimental/Text3D/Source/Text3D/Public/Text3DComponent.h
@@ -67,6 +67,10 @@ public:
 	UPROPERTY(EditAnywhere, BlueprintReadWrite, BlueprintSetter = SetOutline, Category = "Text3D")
 	bool bOutline;
 
+	/** Outline expand/offset amount */
+	UPROPERTY(EditAnywhere, BlueprintReadWrite, BlueprintSetter = SetOutlineExpand, Category = "Text3D", Meta = (EditCondition = "bOutline"))
+	float OutlineExpand;
+
 	/** Material for the front part */
 	UPROPERTY(EditAnywhere, BlueprintReadWrite, BlueprintSetter = SetFrontMaterial, Category = "Text3D")
 	class UMaterialInterface* FrontMaterial;
@@ -127,6 +131,12 @@ public:
 	UPROPERTY(EditAnywhere, BlueprintReadWrite, BlueprintSetter = SetScaleProportionally, Category = "Text3D")
 	bool bScaleProportionally;
 
+	// Lighting flags
+	
+	/** Controls whether the text glyphs should cast a shadow or not. */
+	UPROPERTY(EditAnywhere, BlueprintReadWrite, BlueprintSetter="SetCastShadow", Category=Lighting)
+	bool bCastShadow = true;
+
 	/**
 	 * Delegate called after text is rebuilt
 	 */
@@ -149,6 +159,9 @@ public:
 
 	UFUNCTION(BlueprintCallable, Category = "Rendering|Components|Text3D")
 	void SetOutline(const bool bValue);
+	
+	UFUNCTION(BlueprintCallable, Category = "Rendering|Components|Text3D")
+	void SetOutlineExpand(const float Value);
 
 	/** Set the text extrusion size and signal the primitives to be rebuilt */
 	UFUNCTION(BlueprintCallable, Category = "Rendering|Components|Text3D")
@@ -227,6 +240,10 @@ public:
 	UFUNCTION(BlueprintCallable, Category = "Rendering|Components|Text3D")
 	void SetFreeze(const bool bFreeze);
 
+	/** Changes the value of CastShadow. */
+	UFUNCTION(BlueprintCallable, Category="Rendering")
+	void SetCastShadow(bool NewCastShadow);
+
 	UFUNCTION(BlueprintCallable, Category = "Rendering|Components|Text3D")
 	void GetBounds(FVector& Origin, FVector& BoxExtent);
 
