Text3D performance and duplication fix

diff --git a/Engine/Plugins/Experimental/Text3D/Source/Text3D/Private/MeshCreator.h b/Engine/Plugins/Experimental/Text3D/Source/Text3D/Private/MeshCreator.h
--- a/Engine/Plugins/Experimental/Text3D/Source/Text3D/Private/MeshCreator.h
+++ b/Engine/Plugins/Experimental/Text3D/Source/Text3D/Private/MeshCreator.h
@@ -64,13 +64,13 @@ private:
 	/**
 	 * Recursively compute vertex count in tree.
 	 * @param Node - Tree node.
-	 * @param VertexCount - Pointer to counter.
+	 * @param OutVertexCount - Pointer to counter.
 	 */
 	void AddToVertexCount(const TSharedContourNode& Node, int32& OutVertexCount);
 	/**
 	 * Triangulate solid region with holes and convert contours to old format.
 	 * @param Node - Tree node.
-	 * @param VertexIndex - Pointer to last used vertex index.
+	 * @param OutVertexIndex - Pointer to last used vertex index.
 	 * @param bOutline - Front face has outline and is not filled.
 	 */
 	void TriangulateAndConvert(const TSharedContourNode& Node, int32& OutVertexIndex, const bool bOutline);
diff --git a/Engine/Plugins/Experimental/Text3D/Source/Text3D/Private/Part.cpp b/Engine/Plugins/Experimental/Text3D/Source/Text3D/Private/Part.cpp
--- a/Engine/Plugins/Experimental/Text3D/Source/Text3D/Private/Part.cpp
+++ b/Engine/Plugins/Experimental/Text3D/Source/Text3D/Private/Part.cpp
@@ -84,6 +84,8 @@ bool FPart::ComputeNormal()
 		Normal *= -Scale * FPlatformMath::Sign(FVector2D::CrossProduct(A, C)) / FPlatformMath::Sqrt(NormalLength2);
 	}
 
+	Normal.Normalize();
+
 	return true;
 }
 
diff --git a/Engine/Plugins/Experimental/Text3D/Source/Text3D/Private/Text3DComponent.cpp b/Engine/Plugins/Experimental/Text3D/Source/Text3D/Private/Text3DComponent.cpp
--- a/Engine/Plugins/Experimental/Text3D/Source/Text3D/Private/Text3DComponent.cpp
+++ b/Engine/Plugins/Experimental/Text3D/Source/Text3D/Private/Text3DComponent.cpp
@@ -1,12 +1,13 @@
 // Copyright Epic Games, Inc. All Rights Reserved.
 
 #include "Text3DComponent.h"
+
 #include "Glyph.h"
 #include "MeshCreator.h"
+#include "Text3DEngineSubsystem.h"
 #include "Text3DPrivate.h"
 #include "TextShaper.h"
-
-#include "Text3DEngineSubsystem.h"
+#include "Algo/Count.h"
 #include "Components/StaticMeshComponent.h"
 #include "Engine/Engine.h"
 #include "Engine/Font.h"
@@ -38,9 +39,9 @@ struct FText3DShapedText
 
 using TTextMeshDynamicData = TArray<TUniquePtr<FText3DDynamicData>, TFixedAllocator<static_cast<int32>(EText3DGroupType::TypeCount)>>;
 
-
-UText3DComponent::UText3DComponent() :
-	ShapedText(new FText3DShapedText())
+UText3DComponent::UText3DComponent()
+	: bIsBuilding(false)
+	, ShapedText(new FText3DShapedText())
 {
 	TextRoot = CreateDefaultSubobject<USceneComponent>(TEXT("TextRoot"));
 #if WITH_EDITOR
@@ -97,6 +98,9 @@ void UText3DComponent::OnRegister()
 {
 	Super::OnRegister();
 	TextRoot->AttachToComponent(this, FAttachmentTransformRules::KeepRelativeTransform);
+
+	// Forces rebuild (if build is in progress)
+	bIsBuilding = false;
 	BuildTextMesh();
 }
 
@@ -252,6 +256,75 @@ void UText3DComponent::SetBackMaterial(UMaterialInterface* Value)
 	SetMaterial(EText3DGroupType::Back, Value);
 }
 
+bool UText3DComponent::AllocateGlyphs(int32 Num)
+{
+	int32 DeltaNum = Num - CharacterMeshes.Num();
+	if(DeltaNum == 0)
+	{
+		return false;
+	}
+	// Add characters
+	else if(FMath::Sign(DeltaNum) > 0)
+	{
+		int32 GlyphId = CharacterMeshes.Num() - 1;
+		for(int32 CharacterIndex = 0; CharacterIndex < DeltaNum; ++CharacterIndex)
+		{
+			GlyphId++;
+
+			const FName CharacterKerningComponentName = MakeUniqueObjectName(this, USceneComponent::StaticClass(), FName(*FString::Printf(TEXT("CharacterKerning%d"), GlyphId)));
+			USceneComponent* CharacterKerningComponent = NewObject<USceneComponent>(this, CharacterKerningComponentName);
+			
+#if WITH_EDITOR
+			CharacterKerningComponent->SetIsVisualizationComponent(true);
+#endif
+			CharacterKerningComponent->AttachToComponent(TextRoot, FAttachmentTransformRules::KeepRelativeTransform);
+			CharacterKerningComponent->RegisterComponent();
+			CharacterKernings.Add(CharacterKerningComponent);
+
+			const FName StaticMeshComponentName = MakeUniqueObjectName(this, UStaticMeshComponent::StaticClass(), FName(*FString::Printf(TEXT("StaticMeshComponent%d"), GlyphId)));
+			UStaticMeshComponent* StaticMeshComponent = NewObject<UStaticMeshComponent>(this, StaticMeshComponentName);
+			StaticMeshComponent->RegisterComponent();
+			StaticMeshComponent->SetVisibility(GetVisibleFlag());
+			StaticMeshComponent->SetHiddenInGame(bHiddenInGame);
+			StaticMeshComponent->SetCastShadow(bCastShadow);
+			CharacterMeshes.Add(StaticMeshComponent);
+
+			GetOwner()->AddInstanceComponent(StaticMeshComponent);
+			StaticMeshComponent->AttachToComponent(CharacterKerningComponent, FAttachmentTransformRules::KeepRelativeTransform);
+		}
+	}
+	// Remove characters
+	else
+	{
+		DeltaNum = FMath::Abs(DeltaNum);
+		for(int32 CharacterIndex = CharacterKernings.Num() - 1 - DeltaNum; CharacterIndex < CharacterKernings.Num(); ++CharacterIndex)
+		{
+			USceneComponent* CharacterKerningComponent = CharacterKernings[CharacterIndex];
+			// If called in quick succession, may already be pending destruction
+			if(IsValid(CharacterKerningComponent))
+			{
+				CharacterKerningComponent->DetachFromComponent(FDetachmentTransformRules::KeepRelativeTransform);
+				CharacterKerningComponent->UnregisterComponent();
+				CharacterKerningComponent->DestroyComponent();
+			}
+
+			
+			UStaticMeshComponent* StaticMeshComponent = CharacterMeshes[CharacterIndex];
+			if(IsValid(StaticMeshComponent))
+			{
+				StaticMeshComponent->DetachFromComponent(FDetachmentTransformRules::KeepRelativeTransform);
+				StaticMeshComponent->UnregisterComponent();
+				StaticMeshComponent->DestroyComponent();
+			}
+		}
+		
+		CharacterKernings.RemoveAt(CharacterKernings.Num() - 1 - DeltaNum, DeltaNum);
+		CharacterMeshes.RemoveAt(CharacterMeshes.Num() - 1 - DeltaNum, DeltaNum);
+	}
+
+	return true;
+}
+
 UMaterialInterface* UText3DComponent::GetMaterial(const EText3DGroupType Type) const
 {
 	switch (Type)
@@ -479,12 +552,12 @@ const TArray<UStaticMeshComponent*>& UText3DComponent::GetGlyphMeshComponents()
 	return CharacterMeshes;
 }
 
-void UText3DComponent::Rebuild()
+void UText3DComponent::Rebuild(const bool bCleanCache)
 {
 	bPendingBuild = true;
 	if (!bFreezeBuild)
 	{
-		BuildTextMesh();
+		BuildTextMesh(bCleanCache);
 	}
 }
 void UText3DComponent::CalculateTextWidth()
@@ -633,109 +706,126 @@ void UText3DComponent::ClearTextMesh()
 		}
 	}
 	CharacterKernings.Reset();
+
+	TArray<USceneComponent*> ChildComponents;
+	TextRoot->GetChildrenComponents(true, ChildComponents);
+
+	for (USceneComponent* ChildComponent : ChildComponents)
+	{
+		if(IsValid(ChildComponent))
+		{
+			ChildComponent->DetachFromComponent(FDetachmentTransformRules::KeepRelativeTransform);
+			ChildComponent->DestroyComponent();	
+		}
+	}
 }
 
 void UText3DComponent::BuildTextMesh(const bool bCleanCache)
 {
-	bPendingBuild = false;
-	CheckBevel();
+	bIsBuilding = true;
 
-	ClearTextMesh();
-	if (!Font)
-	{
-		return;
-	}
-	
-	UText3DEngineSubsystem* Subsystem = GEngine->GetEngineSubsystem<UText3DEngineSubsystem>();
-	FCachedFontData& CachedFontData = Subsystem->GetCachedFontData(Font);
-	const FT_Face Face = CachedFontData.GetFreeTypeFace();
-	if (!Face)
+	// Execution guarded by the above flag
+	AsyncTask(ENamedThreads::GameThread, [this, bCleanCache]()
 	{
-		UE_LOG(LogText3D, Error, TEXT("Failed to load font data '%s'"), *CachedFontData.GetFontName());
-		return;
-	}
+		TRACE_CPUPROFILER_EVENT_SCOPE(TEXT("UText3DComponent::Rebuild"));
 
-	CachedCounterReferences.Add(CachedFontData.GetCacheCounter());
-	CachedCounterReferences.Add(CachedFontData.GetMeshesCacheCounter(FGlyphMeshParameters{Extrude, Bevel, BevelType, BevelSegments, bOutline, OutlineExpand}));
+		bPendingBuild = false;
+		CheckBevel();
 
-	ShapedText->Reset();
-	ShapedText->LineHeight = Face->size->metrics.height * FontInverseScale;
-	ShapedText->FontAscender = Face->size->metrics.ascender * FontInverseScale;
-	ShapedText->FontDescender = Face->size->metrics.descender * FontInverseScale;
-	FTextShaper::Get()->ShapeBidirectionalText(Face, Text.ToString(), ShapedText->Lines);
-	
-	CalculateTextWidth();
-	TextRoot->SetRelativeScale3D(GetTextScale());
+		ClearTextMesh();
+		if (!Font)
+		{
+			bIsBuilding = false;
+			return;
+		}
+		
+		UText3DEngineSubsystem* Subsystem = GEngine->GetEngineSubsystem<UText3DEngineSubsystem>();
+		FCachedFontData& CachedFontData = Subsystem->GetCachedFontData(Font);
+		const FT_Face Face = CachedFontData.GetFreeTypeFace();
+		if (!Face)
+		{ 
+			UE_LOG(LogText3D, Error, TEXT("Failed to load font data '%s'"), *CachedFontData.GetFontName());
+			bIsBuilding = false;
+			return;
+		}
 
-	int32 GlyphIndex = 0;
-	for (int32 LineIndex = 0; LineIndex < ShapedText->Lines.Num(); LineIndex++)
-	{
-		const FShapedGlyphLine& ShapedLine = ShapedText->Lines[LineIndex];
-		FVector Location = GetLineLocation(LineIndex);
+		CachedCounterReferences.Add(CachedFontData.GetCacheCounter());
+		CachedCounterReferences.Add(CachedFontData.GetMeshesCacheCounter(FGlyphMeshParameters{Extrude, Bevel, BevelType, BevelSegments, bOutline, OutlineExpand}));
 
-		for (int32 LineGlyph = 0; LineGlyph < ShapedLine.GlyphsToRender.Num(); LineGlyph++)
-		{
-			FVector GlyphLocation = Location;
-			Location.Y += ShapedLine.GetAdvanced(LineGlyph, Kerning, WordSpacing);
+		ShapedText->Reset();
+		ShapedText->LineHeight = Face->size->metrics.height * FontInverseScale;
+		ShapedText->FontAscender = Face->size->metrics.ascender * FontInverseScale;
+		ShapedText->FontDescender = Face->size->metrics.descender * FontInverseScale;
+		FTextShaper::Get()->ShapeBidirectionalText(Face, Text.ToString(), ShapedText->Lines);
+		
+		CalculateTextWidth();
+		TextRoot->SetRelativeScale3D(GetTextScale());
 
-			const FShapedGlyphEntry& ShapedGlyph = ShapedLine.GlyphsToRender[LineGlyph];
-			if (!ShapedGlyph.bIsVisible)
+		// Pre-allocate, avoid new'ing! 
+		AllocateGlyphs(Algo::TransformAccumulate(ShapedText->Lines, [&](const FShapedGlyphLine& ShapedLine)
+		{
+			return Algo::CountIf(ShapedLine.GlyphsToRender, [&](const FShapedGlyphEntry& ShapedGlyph)
 			{
-				continue;
-			}
+				return ShapedGlyph.bIsVisible;
+			});
+		},
+		0));
+
+		int32 GlyphIndex = 0;
+		for (int32 LineIndex = 0; LineIndex < ShapedText->Lines.Num(); LineIndex++)
+		{
+			const FShapedGlyphLine& ShapedLine = ShapedText->Lines[LineIndex];
+			FVector Location = GetLineLocation(LineIndex);
 
-			UStaticMesh* CachedMesh = CachedFontData.GetGlyphMesh(ShapedGlyph.GlyphIndex, FGlyphMeshParameters{Extrude, Bevel, BevelType, BevelSegments, bOutline, OutlineExpand});
-			if (!CachedMesh)
+			for (int32 LineGlyph = 0; LineGlyph < ShapedLine.GlyphsToRender.Num(); LineGlyph++)
 			{
-				continue;
+				FVector GlyphLocation = Location;
+				Location.Y += ShapedLine.GetAdvanced(LineGlyph, Kerning, WordSpacing);
+
+				const FShapedGlyphEntry& ShapedGlyph = ShapedLine.GlyphsToRender[LineGlyph];
+				if (!ShapedGlyph.bIsVisible)
+				{
+					continue;
+				}
+
+				// Count even when mesh is nullptr (allocation still creates components to avoid mesh building in allocation step)
+				const int32 GlyphId = GlyphIndex++;
+
+				UStaticMesh* CachedMesh = CachedFontData.GetGlyphMesh(ShapedGlyph.GlyphIndex, FGlyphMeshParameters{Extrude, Bevel, BevelType, BevelSegments, bOutline, OutlineExpand});
+				if (!CachedMesh || FMath::IsNearlyZero(CachedMesh->GetBounds().SphereRadius))
+				{
+					continue;
+				}
+
+				UStaticMeshComponent* StaticMeshComponent = CharacterMeshes[GlyphId];
+				StaticMeshComponent->SetStaticMesh(CachedMesh);
+				StaticMeshComponent->SetVisibility(GetVisibleFlag());
+				StaticMeshComponent->SetHiddenInGame(bHiddenInGame);
+				StaticMeshComponent->SetCastShadow(bCastShadow);
+
+				FTransform Transform;
+				Transform.SetLocation(GlyphLocation);
+				USceneComponent* CharacterKerningComponent = CharacterKernings[GlyphId];
+				CharacterKerningComponent->SetRelativeTransform(Transform);
 			}
+		}
+		
+		for (int32 Index = 0; Index < static_cast<int32>(EText3DGroupType::TypeCount); Index++)
+		{
+			const EText3DGroupType Type = static_cast<EText3DGroupType>(Index);
+			UpdateMaterial(Type, GetMaterial(Type));
+		}
 
-			const int32 GlyphId = GlyphIndex++;
-
-			const FName CharacterKerningComponentName = MakeUniqueObjectName(this, USceneComponent::StaticClass(), FName(*FString::Printf(TEXT("CharacterKerning%d"), GlyphId)));
-			USceneComponent* CharacterKerningComponent = NewObject<USceneComponent>(this, CharacterKerningComponentName);
-			
-#if WITH_EDITOR
-			CharacterKerningComponent->SetIsVisualizationComponent(true);
-#endif
-			CharacterKerningComponent->AttachToComponent(TextRoot, FAttachmentTransformRules::KeepRelativeTransform);
-			CharacterKerningComponent->RegisterComponent();
-			CharacterKernings.Add(CharacterKerningComponent);
-
-			const FName StaticMeshComponentName = MakeUniqueObjectName(this, UStaticMeshComponent::StaticClass(), FName(*FString::Printf(TEXT("StaticMeshComponent%d"), GlyphId)));
-			UStaticMeshComponent* StaticMeshComponent = NewObject<UStaticMeshComponent>(this, StaticMeshComponentName);
-#if WITH_EDITOR
-			StaticMeshComponent->SetIsVisualizationComponent(true);
-#endif
-			StaticMeshComponent->SetStaticMesh(CachedMesh);
-			StaticMeshComponent->RegisterComponent();
-			StaticMeshComponent->SetVisibility(GetVisibleFlag());
-			StaticMeshComponent->SetHiddenInGame(bHiddenInGame);
-			StaticMeshComponent->SetCastShadow(bCastShadow);
-			CharacterMeshes.Add(StaticMeshComponent);
-
-			GetOwner()->AddInstanceComponent(StaticMeshComponent);
-			StaticMeshComponent->AttachToComponent(CharacterKerningComponent, FAttachmentTransformRules::KeepRelativeTransform);
+		TextGeneratedNativeDelegate.Broadcast();
+		TextGeneratedDelegate.Broadcast();
 
-			FTransform Transform;
-			Transform.SetLocation(GlyphLocation);
-			CharacterKerningComponent->SetRelativeTransform(Transform);
+		if (bCleanCache)
+		{
+			Subsystem->Cleanup();
 		}
-	}
-	
-	for (int32 Index = 0; Index < static_cast<int32>(EText3DGroupType::TypeCount); Index++)
-	{
-		const EText3DGroupType Type = static_cast<EText3DGroupType>(Index);
-		UpdateMaterial(Type, GetMaterial(Type));
-	}
 
-	TextGeneratedNativeDelegate.Broadcast();
-	TextGeneratedDelegate.Broadcast();
-
-	if (bCleanCache)
-	{
-		Subsystem->Cleanup();
-	}
+		bIsBuilding = false;
+	});
 }
 
 void UText3DComponent::CheckBevel()
diff --git a/Engine/Plugins/Experimental/Text3D/Source/Text3D/Private/Text3DEngineSubsystem.cpp b/Engine/Plugins/Experimental/Text3D/Source/Text3D/Private/Text3DEngineSubsystem.cpp
--- a/Engine/Plugins/Experimental/Text3D/Source/Text3D/Private/Text3DEngineSubsystem.cpp
+++ b/Engine/Plugins/Experimental/Text3D/Source/Text3D/Private/Text3DEngineSubsystem.cpp
@@ -209,25 +209,25 @@ UStaticMesh* FCachedFontData::GetGlyphMesh(uint32 GlyphIndex, const FGlyphMeshPa
 	uint32 HashGroup = 0;
 	HashGroup = HashCombine(HashGroup, GetTypeHash(Font));
 	HashGroup = HashCombine(HashGroup, GetTypeHash(GlyphIndex));
-	FString StaticMeshName = FString::Printf(TEXT("Text3D_Char_%u_%u"), HashGroup, HashParameters);
-	
-	FMeshCreator MeshCreator;
-	TSharedContourNode Root = GetGlyphContours(GlyphIndex);
+	const FString StaticMeshName = FString::Printf(TEXT("Text3D_Char_%u_%u"), HashGroup, HashParameters);
+
+	const TSharedContourNode Root = GetGlyphContours(GlyphIndex);
 	if (Root->Children.Num() == 0)
 	{
 		return nullptr;
 	}
 
+	UText3DEngineSubsystem* Subsystem = GEngine->GetEngineSubsystem<UText3DEngineSubsystem>();
+	UStaticMesh* StaticMesh = NewObject<UStaticMesh>(Subsystem, *StaticMeshName);
+	CachedMeshes.Glyphs.Add(GlyphIndex, StaticMesh);
+
+	FMeshCreator MeshCreator;
 	MeshCreator.CreateMeshes(Root, Parameters.Extrude, Parameters.Bevel, Parameters.BevelType, Parameters.BevelSegments, Parameters.bOutline, Parameters.OutlineExpand);
 	MeshCreator.SetFrontAndBevelTextureCoordinates(Parameters.Bevel);
 	MeshCreator.MirrorGroups(Parameters.Extrude);
-
-	UText3DEngineSubsystem* Subsystem = GEngine->GetEngineSubsystem<UText3DEngineSubsystem>();
-	UStaticMesh* StaticMesh = NewObject<UStaticMesh>(Subsystem, *StaticMeshName);
+		
 	MeshCreator.BuildMesh(StaticMesh, Subsystem->DefaultMaterial);
 
-	CachedMeshes.Glyphs.Add(GlyphIndex, StaticMesh);
-
 	return StaticMesh;
 }
 
diff --git a/Engine/Plugins/Experimental/Text3D/Source/Text3D/Public/Glyph.h b/Engine/Plugins/Experimental/Text3D/Source/Text3D/Public/Glyph.h
--- a/Engine/Plugins/Experimental/Text3D/Source/Text3D/Public/Glyph.h
+++ b/Engine/Plugins/Experimental/Text3D/Source/Text3D/Public/Glyph.h
@@ -3,11 +3,9 @@
 #pragma once
 
 #include "Mesh.h"
-#include "MeshTypes.h"
 #include "MeshDescription.h"
 #include "StaticMeshAttributes.h"
 
-
 struct FText3DPolygonGroup
 {
 	int32 FirstVertex;
diff --git a/Engine/Plugins/Experimental/Text3D/Source/Text3D/Public/Text3DActor.h b/Engine/Plugins/Experimental/Text3D/Source/Text3D/Public/Text3DActor.h
--- a/Engine/Plugins/Experimental/Text3D/Source/Text3D/Public/Text3DActor.h
+++ b/Engine/Plugins/Experimental/Text3D/Source/Text3D/Public/Text3DActor.h
@@ -1,10 +1,13 @@
 // Copyright Epic Games, Inc. All Rights Reserved.
 
 #pragma once
+
 #include "CoreMinimal.h"
 #include "GameFramework/Actor.h"
+
 #include "Text3DActor.generated.h"
 
+class UText3DComponent;
 
 UCLASS(ClassGroup = (Text3D), Meta = (DisplayName = "Text 3D", ComponentWrapperClass))
 class TEXT3D_API AText3DActor : public AActor
@@ -15,9 +18,9 @@ public:
 	AText3DActor();
 
 	/** Returns Text3D subobject **/
-	class UText3DComponent* GetText3DComponent() const				{ return Text3DComponent; }
+	UText3DComponent* GetText3DComponent() const { return Text3DComponent; }
 
 private:
 	UPROPERTY(Category = Text, VisibleAnywhere, BlueprintReadOnly, meta = (AllowPrivateAccess = "true"))
-	class UText3DComponent* Text3DComponent;
+	TObjectPtr<UText3DComponent> Text3DComponent;
 };
diff --git a/Engine/Plugins/Experimental/Text3D/Source/Text3D/Public/Text3DComponent.h b/Engine/Plugins/Experimental/Text3D/Source/Text3D/Public/Text3DComponent.h
--- a/Engine/Plugins/Experimental/Text3D/Source/Text3D/Public/Text3DComponent.h
+++ b/Engine/Plugins/Experimental/Text3D/Source/Text3D/Public/Text3DComponent.h
@@ -1,14 +1,13 @@
 // Copyright Epic Games, Inc. All Rights Reserved.
 
 #pragma once
-#include "Mesh.h"
-#include "BevelType.h"
 
 #include "CoreMinimal.h"
+#include "BevelType.h"
+#include "Mesh.h"
 #include "UObject/ObjectMacros.h"
-#include "UObject/TextProperty.h"
-#include "Text3DComponent.generated.h"
 
+#include "Text3DComponent.generated.h"
 
 UENUM()
 enum class EText3DVerticalTextAlignment : uint8
@@ -146,8 +145,7 @@ public:
 	 * Delegate called after text is rebuilt
 	 */
 	DECLARE_MULTICAST_DELEGATE(FTextGeneratedNative);
-	FTextGeneratedNative& OnTextGenerated()				{ return TextGeneratedNativeDelegate; }
-
+	FTextGeneratedNative& OnTextGenerated() { return TextGeneratedNativeDelegate; }
 
 	/** Set the text value and signal the primitives to be rebuilt */
 	UFUNCTION(BlueprintCallable, Category = "Rendering|Components|Text3D")
@@ -280,6 +278,9 @@ private:
 
 	FTextGeneratedNative TextGeneratedNativeDelegate;
 
+	/** Flagged as true when the text mesh is being built. */
+	std::atomic<bool> bIsBuilding;
+	
 	bool bPendingBuild;
 	bool bFreezeBuild;
 
@@ -292,12 +293,14 @@ private:
 	UPROPERTY(Transient)
 	TArray<UStaticMeshComponent*> CharacterMeshes;
 
+	/** Allocates, or shrinks existing components to match the input number. Returns false if nothing modified. */
+	bool AllocateGlyphs(int32 Num);
+
 	class UMaterialInterface* GetMaterial(const EText3DGroupType Type) const;
 	void SetMaterial(const EText3DGroupType Type, class UMaterialInterface* Material);
 
-
-	void Rebuild();
 	void UpdateTransforms();
+	void Rebuild(const bool bCleanCache = false);
 	void ClearTextMesh();
 	void BuildTextMesh(const bool bCleanCache = false);
 	void CheckBevel();
