Fix crash in Texture3DComponent introduced in CL-19059943.
We need to ensure that our component is still valid before handling the async callback.

#[review] [REDACTED_HANDLE] [REDACTED_HANDLE] [REDACTED_HANDLE]

diff --git a/Engine/Plugins/Experimental/Text3D/Source/Text3D/Private/Text3DComponent.cpp b/Engine/Plugins/Experimental/Text3D/Source/Text3D/Private/Text3DComponent.cpp
--- a/Engine/Plugins/Experimental/Text3D/Source/Text3D/Private/Text3DComponent.cpp
+++ b/Engine/Plugins/Experimental/Text3D/Source/Text3D/Private/Text3DComponent.cpp
@@ -722,110 +722,126 @@ void UText3DComponent::ClearTextMesh()
 
 void UText3DComponent::BuildTextMesh(const bool bCleanCache)
 {
+	// If we're already building, or have a build pending, don't do anything.
+	if (bIsBuilding)
+	{
+		return;
+	}
+	
 	bIsBuilding = true;
 
+	TWeakObjectPtr<UText3DComponent> WeakThis(this);
+	
 	// Execution guarded by the above flag
-	AsyncTask(ENamedThreads::GameThread, [this, bCleanCache]()
+	AsyncTask(ENamedThreads::GameThread, [WeakThis, bCleanCache]()
 	{
-		TRACE_CPUPROFILER_EVENT_SCOPE(TEXT("UText3DComponent::Rebuild"));
-
-		bPendingBuild = false;
-		CheckBevel();
-
-		ClearTextMesh();
-		if (!Font)
+		if (UText3DComponent* StrongThis = WeakThis.Get())
 		{
-			bIsBuilding = false;
-			return;
-		}
-		
-		UText3DEngineSubsystem* Subsystem = GEngine->GetEngineSubsystem<UText3DEngineSubsystem>();
-		FCachedFontData& CachedFontData = Subsystem->GetCachedFontData(Font);
-		const FT_Face Face = CachedFontData.GetFreeTypeFace();
-		if (!Face)
-		{ 
-			UE_LOG(LogText3D, Error, TEXT("Failed to load font data '%s'"), *CachedFontData.GetFontName());
-			bIsBuilding = false;
-			return;
+			StrongThis->BuildTextMeshInternal(bCleanCache);
 		}
+	});
+}
 
-		CachedCounterReferences.Add(CachedFontData.GetCacheCounter());
-		CachedCounterReferences.Add(CachedFontData.GetMeshesCacheCounter(FGlyphMeshParameters{Extrude, Bevel, BevelType, BevelSegments, bOutline, OutlineExpand}));
+void UText3DComponent::BuildTextMeshInternal(const bool bCleanCache)
+{
+	TRACE_CPUPROFILER_EVENT_SCOPE(TEXT("UText3DComponent::Rebuild"));
 
-		ShapedText->Reset();
-		ShapedText->LineHeight = Face->size->metrics.height * FontInverseScale;
-		ShapedText->FontAscender = Face->size->metrics.ascender * FontInverseScale;
-		ShapedText->FontDescender = Face->size->metrics.descender * FontInverseScale;
-		FTextShaper::Get()->ShapeBidirectionalText(Face, Text.ToString(), ShapedText->Lines);
-		
-		CalculateTextWidth();
-		TextRoot->SetRelativeScale3D(GetTextScale());
+	bPendingBuild = false;
+	CheckBevel();
 
-		// Pre-allocate, avoid new'ing! 
-		AllocateGlyphs(Algo::TransformAccumulate(ShapedText->Lines, [&](const FShapedGlyphLine& ShapedLine)
+	ClearTextMesh();
+	if (!Font)
+	{
+		bIsBuilding = false;
+		return;
+	}
+	
+	UText3DEngineSubsystem* Subsystem = GEngine->GetEngineSubsystem<UText3DEngineSubsystem>();
+	FCachedFontData& CachedFontData = Subsystem->GetCachedFontData(Font);
+	const FT_Face Face = CachedFontData.GetFreeTypeFace();
+	if (!Face)
+	{ 
+		UE_LOG(LogText3D, Error, TEXT("Failed to load font data '%s'"), *CachedFontData.GetFontName());
+		bIsBuilding = false;
+		return;
+	}
+
+	CachedCounterReferences.Add(CachedFontData.GetCacheCounter());
+	CachedCounterReferences.Add(CachedFontData.GetMeshesCacheCounter(FGlyphMeshParameters{Extrude, Bevel, BevelType, BevelSegments, bOutline, OutlineExpand}));
+
+	ShapedText->Reset();
+	ShapedText->LineHeight = Face->size->metrics.height * FontInverseScale;
+	ShapedText->FontAscender = Face->size->metrics.ascender * FontInverseScale;
+	ShapedText->FontDescender = Face->size->metrics.descender * FontInverseScale;
+	FTextShaper::Get()->ShapeBidirectionalText(Face, Text.ToString(), ShapedText->Lines);
+	
+	CalculateTextWidth();
+	TextRoot->SetRelativeScale3D(GetTextScale());
+
+	// Pre-allocate, avoid new'ing! 
+	AllocateGlyphs(Algo::TransformAccumulate(ShapedText->Lines, [&](const FShapedGlyphLine& ShapedLine)
+	{
+		return Algo::CountIf(ShapedLine.GlyphsToRender, [&](const FShapedGlyphEntry& ShapedGlyph)
 		{
-			return Algo::CountIf(ShapedLine.GlyphsToRender, [&](const FShapedGlyphEntry& ShapedGlyph)
-			{
-				return ShapedGlyph.bIsVisible;
-			});
-		},
-		0));
+			return ShapedGlyph.bIsVisible;
+		});
+	},
+	0));
 
-		int32 GlyphIndex = 0;
-		for (int32 LineIndex = 0; LineIndex < ShapedText->Lines.Num(); LineIndex++)
+	int32 GlyphIndex = 0;
+	for (int32 LineIndex = 0; LineIndex < ShapedText->Lines.Num(); LineIndex++)
+	{
+		const FShapedGlyphLine& ShapedLine = ShapedText->Lines[LineIndex];
+		FVector Location = GetLineLocation(LineIndex);
+
+		for (int32 LineGlyph = 0; LineGlyph < ShapedLine.GlyphsToRender.Num(); LineGlyph++)
 		{
-			const FShapedGlyphLine& ShapedLine = ShapedText->Lines[LineIndex];
-			FVector Location = GetLineLocation(LineIndex);
+			FVector GlyphLocation = Location;
+			Location.Y += ShapedLine.GetAdvanced(LineGlyph, Kerning, WordSpacing);
 
-			for (int32 LineGlyph = 0; LineGlyph < ShapedLine.GlyphsToRender.Num(); LineGlyph++)
+			const FShapedGlyphEntry& ShapedGlyph = ShapedLine.GlyphsToRender[LineGlyph];
+			if (!ShapedGlyph.bIsVisible)
 			{
-				FVector GlyphLocation = Location;
-				Location.Y += ShapedLine.GetAdvanced(LineGlyph, Kerning, WordSpacing);
-
-				const FShapedGlyphEntry& ShapedGlyph = ShapedLine.GlyphsToRender[LineGlyph];
-				if (!ShapedGlyph.bIsVisible)
-				{
-					continue;
-				}
-
-				// Count even when mesh is nullptr (allocation still creates components to avoid mesh building in allocation step)
-				const int32 GlyphId = GlyphIndex++;
-
-				UStaticMesh* CachedMesh = CachedFontData.GetGlyphMesh(ShapedGlyph.GlyphIndex, FGlyphMeshParameters{Extrude, Bevel, BevelType, BevelSegments, bOutline, OutlineExpand});
-				if (!CachedMesh || FMath::IsNearlyZero(CachedMesh->GetBounds().SphereRadius))
-				{
-					continue;
-				}
-
-				UStaticMeshComponent* StaticMeshComponent = CharacterMeshes[GlyphId];
-				StaticMeshComponent->SetStaticMesh(CachedMesh);
-				StaticMeshComponent->SetVisibility(GetVisibleFlag());
-				StaticMeshComponent->SetHiddenInGame(bHiddenInGame);
-				StaticMeshComponent->SetCastShadow(bCastShadow);
-
-				FTransform Transform;
-				Transform.SetLocation(GlyphLocation);
-				USceneComponent* CharacterKerningComponent = CharacterKernings[GlyphId];
-				CharacterKerningComponent->SetRelativeTransform(Transform);
+				continue;
 			}
-		}
-		
-		for (int32 Index = 0; Index < static_cast<int32>(EText3DGroupType::TypeCount); Index++)
-		{
-			const EText3DGroupType Type = static_cast<EText3DGroupType>(Index);
-			UpdateMaterial(Type, GetMaterial(Type));
-		}
 
-		TextGeneratedNativeDelegate.Broadcast();
-		TextGeneratedDelegate.Broadcast();
+			// Count even when mesh is nullptr (allocation still creates components to avoid mesh building in allocation step)
+			const int32 GlyphId = GlyphIndex++;
 
-		if (bCleanCache)
-		{
-			Subsystem->Cleanup();
+			UStaticMesh* CachedMesh = CachedFontData.GetGlyphMesh(ShapedGlyph.GlyphIndex, FGlyphMeshParameters{Extrude, Bevel, BevelType, BevelSegments, bOutline, OutlineExpand});
+			if (!CachedMesh || FMath::IsNearlyZero(CachedMesh->GetBounds().SphereRadius))
+			{
+				continue;
+			}
+
+			UStaticMeshComponent* StaticMeshComponent = CharacterMeshes[GlyphId];
+			StaticMeshComponent->SetStaticMesh(CachedMesh);
+			StaticMeshComponent->SetVisibility(GetVisibleFlag());
+			StaticMeshComponent->SetHiddenInGame(bHiddenInGame);
+			StaticMeshComponent->SetCastShadow(bCastShadow);
+
+			FTransform Transform;
+			Transform.SetLocation(GlyphLocation);
+			USceneComponent* CharacterKerningComponent = CharacterKernings[GlyphId];
+			CharacterKerningComponent->SetRelativeTransform(Transform);
 		}
+	}
+	
+	for (int32 Index = 0; Index < static_cast<int32>(EText3DGroupType::TypeCount); Index++)
+	{
+		const EText3DGroupType Type = static_cast<EText3DGroupType>(Index);
+		UpdateMaterial(Type, GetMaterial(Type));
+	}
 
-		bIsBuilding = false;
-	});
+	TextGeneratedNativeDelegate.Broadcast();
+	TextGeneratedDelegate.Broadcast();
+
+	if (bCleanCache)
+	{
+		Subsystem->Cleanup();
+	}
+
+	bIsBuilding = false;
 }
 
 void UText3DComponent::CheckBevel()
diff --git a/Engine/Plugins/Experimental/Text3D/Source/Text3D/Public/Text3DComponent.h b/Engine/Plugins/Experimental/Text3D/Source/Text3D/Public/Text3DComponent.h
--- a/Engine/Plugins/Experimental/Text3D/Source/Text3D/Public/Text3DComponent.h
+++ b/Engine/Plugins/Experimental/Text3D/Source/Text3D/Public/Text3DComponent.h
@@ -303,6 +303,7 @@ private:
 	void Rebuild(const bool bCleanCache = false);
 	void ClearTextMesh();
 	void BuildTextMesh(const bool bCleanCache = false);
+	void BuildTextMeshInternal(const bool bCleanCache);
 	void CheckBevel();
 	float MaxBevel() const;
 
