Text3D Horde Issue #145196

diff --git a/Engine/Plugins/Experimental/Text3D/Source/Text3D/Private/Text3DComponent.cpp b/Engine/Plugins/Experimental/Text3D/Source/Text3D/Private/Text3DComponent.cpp
--- a/Engine/Plugins/Experimental/Text3D/Source/Text3D/Private/Text3DComponent.cpp
+++ b/Engine/Plugins/Experimental/Text3D/Source/Text3D/Private/Text3DComponent.cpp
@@ -8,9 +8,11 @@
 #include "Text3DPrivate.h"
 #include "TextShaper.h"
 #include "Algo/Count.h"
+#include "Async/Async.h"
 #include "Components/StaticMeshComponent.h"
 #include "Engine/Engine.h"
 #include "Engine/Font.h"
+#include "Engine/StaticMesh.h"
 #include "Materials/Material.h"
 #include "UObject/ConstructorHelpers.h"
 
