Mark the SceneComponents and StaticMeshComponents used for rendering text as transient, so they don't inadvertently dirty Level packages.

diff --git a/Engine/Plugins/Experimental/Text3D/Source/Text3D/Private/Text3DComponent.cpp b/Engine/Plugins/Experimental/Text3D/Source/Text3D/Private/Text3DComponent.cpp
--- a/Engine/Plugins/Experimental/Text3D/Source/Text3D/Private/Text3DComponent.cpp
+++ b/Engine/Plugins/Experimental/Text3D/Source/Text3D/Private/Text3DComponent.cpp
@@ -274,7 +274,7 @@ bool UText3DComponent::AllocateGlyphs(int32 Num)
 			GlyphId++;
 
 			const FName CharacterKerningComponentName = MakeUniqueObjectName(this, USceneComponent::StaticClass(), FName(*FString::Printf(TEXT("CharacterKerning%d"), GlyphId)));
-			USceneComponent* CharacterKerningComponent = NewObject<USceneComponent>(this, CharacterKerningComponentName);
+			USceneComponent* CharacterKerningComponent = NewObject<USceneComponent>(this, CharacterKerningComponentName, RF_Transient);
 			
 #if WITH_EDITOR
 			CharacterKerningComponent->SetIsVisualizationComponent(true);
@@ -284,7 +284,7 @@ bool UText3DComponent::AllocateGlyphs(int32 Num)
 			CharacterKernings.Add(CharacterKerningComponent);
 
 			const FName StaticMeshComponentName = MakeUniqueObjectName(this, UStaticMeshComponent::StaticClass(), FName(*FString::Printf(TEXT("StaticMeshComponent%d"), GlyphId)));
-			UStaticMeshComponent* StaticMeshComponent = NewObject<UStaticMeshComponent>(this, StaticMeshComponentName);
+			UStaticMeshComponent* StaticMeshComponent = NewObject<UStaticMeshComponent>(this, StaticMeshComponentName, RF_Transient);
 			StaticMeshComponent->RegisterComponent();
 			StaticMeshComponent->SetVisibility(GetVisibleFlag());
 			StaticMeshComponent->SetHiddenInGame(bHiddenInGame);
