Text3D: Fixed crashes caused by the async build executing after the component has been unregistered

diff --git a/Engine/Plugins/Experimental/Text3D/Source/Text3D/Private/Text3DComponent.cpp b/Engine/Plugins/Experimental/Text3D/Source/Text3D/Private/Text3DComponent.cpp
--- a/Engine/Plugins/Experimental/Text3D/Source/Text3D/Private/Text3DComponent.cpp
+++ b/Engine/Plugins/Experimental/Text3D/Source/Text3D/Private/Text3DComponent.cpp
@@ -14,6 +14,7 @@
 #include "Engine/Font.h"
 #include "Engine/StaticMesh.h"
 #include "Materials/Material.h"
+#include "Misc/ScopeExit.h"
 #include "UObject/ConstructorHelpers.h"
 
 #define LOCTEXT_NAMESPACE "Text3D"
@@ -748,13 +749,19 @@ void UText3DComponent::BuildTextMeshInternal(const bool bCleanCache)
 {
 	TRACE_CPUPROFILER_EVENT_SCOPE(TEXT("UText3DComponent::Rebuild"));
 
+	ON_SCOPE_EXIT { bIsBuilding = false; };
+
+	if (!IsRegistered())
+	{
+		return;
+	}
+
 	bPendingBuild = false;
 	CheckBevel();
 
 	ClearTextMesh();
 	if (!Font)
 	{
-		bIsBuilding = false;
 		return;
 	}
 	
@@ -764,7 +771,6 @@ void UText3DComponent::BuildTextMeshInternal(const bool bCleanCache)
 	if (!Face)
 	{ 
 		UE_LOG(LogText3D, Error, TEXT("Failed to load font data '%s'"), *CachedFontData.GetFontName());
-		bIsBuilding = false;
 		return;
 	}
 
@@ -842,8 +848,6 @@ void UText3DComponent::BuildTextMeshInternal(const bool bCleanCache)
 	{
 		Subsystem->Cleanup();
 	}
-
-	bIsBuilding = false;
 }
 
 void UText3DComponent::CheckBevel()
