Text3D

- Expose GetTextScale as public
- Add ComputeTextScale method

diff --git a/Engine/Plugins/Experimental/Text3D/Source/Text3D/Private/Text3DComponent.cpp b/Engine/Plugins/Experimental/Text3D/Source/Text3D/Private/Text3DComponent.cpp
--- a/Engine/Plugins/Experimental/Text3D/Source/Text3D/Private/Text3DComponent.cpp
+++ b/Engine/Plugins/Experimental/Text3D/Source/Text3D/Private/Text3DComponent.cpp
@@ -95,6 +95,8 @@ UText3DComponent::UText3DComponent()
 
 	bPendingBuild = false;
 	bFreezeBuild = false;
+
+	TextScale = FVector::ZeroVector;
 }
 
 void UText3DComponent::OnRegister()
@@ -571,12 +573,12 @@ void UText3DComponent::CalculateTextWidth()
 	}
 }
 
-float UText3DComponent::GetTextHeight()
+float UText3DComponent::GetTextHeight() const
 {
 	return ShapedText->Lines.Num() * ShapedText->LineHeight + (ShapedText->Lines.Num() - 1) * LineSpacing;
 }
 
-FVector UText3DComponent::GetTextScale()
+void UText3DComponent::CalculateTextScale()
 {
 	FVector Scale(1.0f, 1.0f, 1.0f);
 
@@ -610,7 +612,17 @@ FVector UText3DComponent::GetTextScale()
 		Scale.X = Scale.Y;
 	}
 
-	return Scale;
+	TextScale = Scale;
+}
+
+FVector UText3DComponent::GetTextScale()
+{
+	if (TextScale == FVector::ZeroVector)
+	{
+		CalculateTextScale();
+	}
+
+	return TextScale;
 }
 
 FVector UText3DComponent::GetLineLocation(int32 LineIndex)
@@ -656,6 +668,7 @@ FVector UText3DComponent::GetLineLocation(int32 LineIndex)
 void UText3DComponent::UpdateTransforms()
 {
 	CalculateTextWidth();
+	CalculateTextScale();
 	const FVector Scale = GetTextScale();
 	TextRoot->SetRelativeScale3D(Scale);
 
@@ -784,6 +797,7 @@ void UText3DComponent::BuildTextMeshInternal(const bool bCleanCache)
 	FTextShaper::Get()->ShapeBidirectionalText(Face, Text.ToString(), ShapedText->Lines);
 	
 	CalculateTextWidth();
+	CalculateTextScale();
 	TextRoot->SetRelativeScale3D(GetTextScale());
 
 	// Pre-allocate, avoid new'ing! 
diff --git a/Engine/Plugins/Experimental/Text3D/Source/Text3D/Public/Text3DComponent.h b/Engine/Plugins/Experimental/Text3D/Source/Text3D/Public/Text3DComponent.h
--- a/Engine/Plugins/Experimental/Text3D/Source/Text3D/Public/Text3DComponent.h
+++ b/Engine/Plugins/Experimental/Text3D/Source/Text3D/Public/Text3DComponent.h
@@ -264,6 +264,10 @@ public:
 	/** Gets all the glyph meshes*/
 	UFUNCTION(BlueprintCallable, Category = "Rendering|Components|Text3D")
 	const TArray<UStaticMeshComponent*>& GetGlyphMeshComponents();
+	
+	/** Gets the scale of actual text geometry, taking into account MaxWidth and MaxHeight constraints. This function will NOT return the component scale*/
+	UFUNCTION(BlueprintCallable, Category = "Rendering|Components|Text3D")
+	FVector GetTextScale();
 
 protected:
 	virtual void OnVisibilityChanged() override;
@@ -284,6 +288,8 @@ private:
 	bool bPendingBuild;
 	bool bFreezeBuild;
 
+	FVector TextScale;
+
 	TSharedPtr<struct FText3DShapedText> ShapedText;
 	TArray<TSharedPtr<int32>> CachedCounterReferences;
 
@@ -310,7 +316,7 @@ private:
 	void UpdateMaterial(const EText3DGroupType Type, class UMaterialInterface* Material);
 
 	void CalculateTextWidth();
-	float GetTextHeight();
-	FVector GetTextScale();
+	float GetTextHeight() const;
+	void CalculateTextScale();
 	FVector GetLineLocation(int32 LineIndex);
 };
