Text3D building when saving speculative fix

diff --git a/Engine/Plugins/Experimental/Text3D/Source/Text3D/Private/Text3DComponent.cpp b/Engine/Plugins/Experimental/Text3D/Source/Text3D/Private/Text3DComponent.cpp
--- a/Engine/Plugins/Experimental/Text3D/Source/Text3D/Private/Text3DComponent.cpp
+++ b/Engine/Plugins/Experimental/Text3D/Source/Text3D/Private/Text3DComponent.cpp
@@ -565,6 +565,7 @@ void UText3DComponent::Rebuild(const bool bCleanCache)
 		BuildTextMesh(bCleanCache);
 	}
 }
+
 void UText3DComponent::CalculateTextWidth()
 {
 	for (FShapedGlyphLine& ShapedLine : ShapedText->Lines)
@@ -747,13 +748,16 @@ void UText3DComponent::BuildTextMesh(const bool bCleanCache)
 	bIsBuilding = true;
 
 	TWeakObjectPtr<UText3DComponent> WeakThis(this);
-	
+
 	// Execution guarded by the above flag
 	AsyncTask(ENamedThreads::GameThread, [WeakThis, bCleanCache]()
 	{
 		if (UText3DComponent* StrongThis = WeakThis.Get())
 		{
-			StrongThis->BuildTextMeshInternal(bCleanCache);
+			if(!UE::IsSavingPackage(StrongThis))
+			{
+				StrongThis->BuildTextMeshInternal(bCleanCache);	
+			}
 		}
 	});
 }
@@ -800,7 +804,7 @@ void UText3DComponent::BuildTextMeshInternal(const bool bCleanCache)
 	CalculateTextScale();
 	TextRoot->SetRelativeScale3D(GetTextScale());
 
-	// Pre-allocate, avoid new'ing! 
+	// Pre-allocate, avoid new'ing!
 	AllocateGlyphs(Algo::TransformAccumulate(ShapedText->Lines, [&](const FShapedGlyphLine& ShapedLine)
 	{
 		return Algo::CountIf(ShapedLine.GlyphsToRender, [&](const FShapedGlyphEntry& ShapedGlyph)
