Text3D self intersection fix

diff --git a/Engine/Plugins/Experimental/Text3D/Source/Text3D/Private/MeshCreator.cpp b/Engine/Plugins/Experimental/Text3D/Source/Text3D/Private/MeshCreator.cpp
--- a/Engine/Plugins/Experimental/Text3D/Source/Text3D/Private/MeshCreator.cpp
+++ b/Engine/Plugins/Experimental/Text3D/Source/Text3D/Private/MeshCreator.cpp
@@ -392,7 +392,7 @@ void FMeshCreator::TriangulateAndConvert(const TSharedContourNode& Node, int32&
 		};
 
 
-		// Outter
+		// Outer
 		ProcessContour(Node);
 
 		// Holes
@@ -401,7 +401,6 @@ void FMeshCreator::TriangulateAndConvert(const TSharedContourNode& Node, int32&
 			ProcessContour(Child);
 		}
 
-
 		if (!bOutline)
 		{
 			Triangulation.Triangulate();
diff --git a/Engine/Plugins/Experimental/Text3D/Source/Text3D/Private/Text3DComponent.cpp b/Engine/Plugins/Experimental/Text3D/Source/Text3D/Private/Text3DComponent.cpp
--- a/Engine/Plugins/Experimental/Text3D/Source/Text3D/Private/Text3DComponent.cpp
+++ b/Engine/Plugins/Experimental/Text3D/Source/Text3D/Private/Text3DComponent.cpp
@@ -840,16 +840,32 @@ void UText3DComponent::BuildTextMeshInternal(const bool bCleanCache)
 				continue;
 			}
 
-			UStaticMeshComponent* StaticMeshComponent = CharacterMeshes[GlyphId];
-			StaticMeshComponent->SetStaticMesh(CachedMesh);
-			StaticMeshComponent->SetVisibility(GetVisibleFlag());
-			StaticMeshComponent->SetHiddenInGame(bHiddenInGame);
-			StaticMeshComponent->SetCastShadow(bCastShadow);
+			if(CharacterMeshes.IsValidIndex(GlyphId))
+			{
+				UStaticMeshComponent* StaticMeshComponent = CharacterMeshes[GlyphId];
+				StaticMeshComponent->SetStaticMesh(CachedMesh);			
+				StaticMeshComponent->SetVisibility(GetVisibleFlag());
+				StaticMeshComponent->SetHiddenInGame(bHiddenInGame);
+				StaticMeshComponent->SetCastShadow(bCastShadow);
+			}
+			else
+			{
+				// @note: This shouldn't occur, but it does under unknown circumstances (UE-164789) so it should be handled 
+				UE_LOG(LogText3D, Error, TEXT("CharacterMesh not found at index %d"), GlyphId);
+			}
 
-			FTransform Transform;
-			Transform.SetLocation(GlyphLocation);
-			USceneComponent* CharacterKerningComponent = CharacterKernings[GlyphId];
-			CharacterKerningComponent->SetRelativeTransform(Transform);
+			if(CharacterKernings.IsValidIndex(GlyphId))
+			{
+				FTransform Transform;
+				Transform.SetLocation(GlyphLocation);
+				USceneComponent* CharacterKerningComponent = CharacterKernings[GlyphId];
+				CharacterKerningComponent->SetRelativeTransform(Transform);
+			}
+			else
+			{
+				// @note: This shouldn't occur, but it does under unknown circumstances (UE-164789) so it should be handled
+				UE_LOG(LogText3D, Error, TEXT("CharacterKerning not found at index %d"), GlyphId);
+			}
 		}
 	}
 	
diff --git a/Engine/Plugins/Experimental/Text3D/Source/Text3D/Private/Text3DEngineSubsystem.cpp b/Engine/Plugins/Experimental/Text3D/Source/Text3D/Private/Text3DEngineSubsystem.cpp
--- a/Engine/Plugins/Experimental/Text3D/Source/Text3D/Private/Text3DEngineSubsystem.cpp
+++ b/Engine/Plugins/Experimental/Text3D/Source/Text3D/Private/Text3DEngineSubsystem.cpp
@@ -245,19 +245,14 @@ TSharedContourNode FCachedFontData::GetGlyphContours(uint32 GlyphIndex)
 {
 	check(FreeTypeFace);
 
-	//if (Glyphs.Contains(GlyphIndex))
-	//{
-	//	return Glyphs[GlyphIndex];
-	//}
-
-	if (FT_Load_Glyph(FreeTypeFace, GlyphIndex, FT_LOAD_DEFAULT))
+	if (FT_Load_Glyph(FreeTypeFace, GlyphIndex, FT_LOAD_NO_HINTING | FT_LOAD_NO_BITMAP))
 	{
 		return nullptr;
 	}
 
 	FGlyphLoader GlyphLoader(FreeTypeFace->glyph);
 	TSharedContourNode Root = GlyphLoader.GetContourList();
-	//Glyphs.Add(GlyphIndex, Root);
+
 	return Root;
 }
 
