Text3D: Fix CIS

diff --git a/Engine/Plugins/Experimental/Text3D/Source/Text3D/Public/Glyph.h b/Engine/Plugins/Experimental/Text3D/Source/Text3D/Public/Glyph.h
--- a/Engine/Plugins/Experimental/Text3D/Source/Text3D/Public/Glyph.h
+++ b/Engine/Plugins/Experimental/Text3D/Source/Text3D/Public/Glyph.h
@@ -19,7 +19,7 @@ class FText3DGlyph
 public:
 	FText3DGlyph();
 
-	void Build(UStaticMesh* StaticMesh, class UMaterial* DefaultMaterial);
+	void Build(class UStaticMesh* StaticMesh, class UMaterial* DefaultMaterial);
 
 	FMeshDescription& GetMeshDescription();
 	FStaticMeshAttributes& GetStaticMeshAttributes();
