Text3D: Fix non-unity CIS issue.

diff --git a/Engine/Plugins/Experimental/Text3D/Source/Text3D/Private/MeshCreator.h b/Engine/Plugins/Experimental/Text3D/Source/Text3D/Private/MeshCreator.h
--- a/Engine/Plugins/Experimental/Text3D/Source/Text3D/Private/MeshCreator.h
+++ b/Engine/Plugins/Experimental/Text3D/Source/Text3D/Private/MeshCreator.h
@@ -27,7 +27,7 @@ public:
 	void CreateMeshes(const TSharedContourNode& Root, const float Extrude, const float Bevel, const EText3DBevelType Type, const int32 BevelSegments, const bool bOutline, const float OutlineExpand = 0.5f);
 	void SetFrontAndBevelTextureCoordinates(const float Bevel);
 	void MirrorGroups(const float Extrude);
-	void BuildMesh(UStaticMesh* StaticMesh, class UMaterial* DefaultMaterial);
+	void BuildMesh(class UStaticMesh* StaticMesh, class UMaterial* DefaultMaterial);
 
 private:
 	TSharedPtr<class FText3DGlyph> Glyph;
