Text3D glyph allocation cleanup

diff --git a/Engine/Plugins/Experimental/Text3D/Source/Text3D/Private/Text3DComponent.cpp b/Engine/Plugins/Experimental/Text3D/Source/Text3D/Private/Text3DComponent.cpp
--- a/Engine/Plugins/Experimental/Text3D/Source/Text3D/Private/Text3DComponent.cpp
+++ b/Engine/Plugins/Experimental/Text3D/Source/Text3D/Private/Text3DComponent.cpp
@@ -264,12 +264,12 @@ void UText3DComponent::SetBackMaterial(UMaterialInterface* Value)
 bool UText3DComponent::AllocateGlyphs(int32 Num)
 {
 	int32 DeltaNum = Num - CharacterMeshes.Num();
-	if(DeltaNum == 0)
+	if (DeltaNum == 0)
 	{
 		return false;
 	}
 	// Add characters
-	else if(FMath::Sign(DeltaNum) > 0)
+	if (FMath::Sign(DeltaNum) > 0)
 	{
 		int32 GlyphId = CharacterMeshes.Num() - 1;
 		for(int32 CharacterIndex = 0; CharacterIndex < DeltaNum; ++CharacterIndex)
@@ -301,25 +301,28 @@ bool UText3DComponent::AllocateGlyphs(int32 Num)
 	// Remove characters
 	else
 	{
+		AActor* OwnerActor = GetOwner();
 		DeltaNum = FMath::Abs(DeltaNum);
 		for(int32 CharacterIndex = CharacterKernings.Num() - 1 - DeltaNum; CharacterIndex < CharacterKernings.Num(); ++CharacterIndex)
 		{
 			USceneComponent* CharacterKerningComponent = CharacterKernings[CharacterIndex];
 			// If called in quick succession, may already be pending destruction
-			if(IsValid(CharacterKerningComponent))
+			if (IsValid(CharacterKerningComponent))
 			{
 				CharacterKerningComponent->DetachFromComponent(FDetachmentTransformRules::KeepRelativeTransform);
 				CharacterKerningComponent->UnregisterComponent();
 				CharacterKerningComponent->DestroyComponent();
+				OwnerActor->RemoveInstanceComponent(CharacterKerningComponent);
 			}
 
 			
 			UStaticMeshComponent* StaticMeshComponent = CharacterMeshes[CharacterIndex];
-			if(IsValid(StaticMeshComponent))
+			if (IsValid(StaticMeshComponent))
 			{
 				StaticMeshComponent->DetachFromComponent(FDetachmentTransformRules::KeepRelativeTransform);
 				StaticMeshComponent->UnregisterComponent();
 				StaticMeshComponent->DestroyComponent();
+				OwnerActor->RemoveInstanceComponent(StaticMeshComponent);
 			}
 		}
 		
@@ -509,7 +512,7 @@ void UText3DComponent::SetFreeze(const bool bFreeze)
 
 void UText3DComponent::SetCastShadow(bool NewCastShadow)
 {
-	if(NewCastShadow != bCastShadow)
+	if (NewCastShadow != bCastShadow)
 	{
 		bCastShadow = NewCastShadow;
 
@@ -724,17 +727,24 @@ void UText3DComponent::ClearTextMesh()
 	}
 	CharacterKernings.Reset();
 
+	constexpr bool bIncludeChildDescendants = true;
 	TArray<USceneComponent*> ChildComponents;
-	TextRoot->GetChildrenComponents(true, ChildComponents);
+	TextRoot->GetChildrenComponents(bIncludeChildDescendants, ChildComponents);
 
 	for (USceneComponent* ChildComponent : ChildComponents)
 	{
-		if(IsValid(ChildComponent))
+		if (IsValid(ChildComponent))
 		{
 			ChildComponent->DetachFromComponent(FDetachmentTransformRules::KeepRelativeTransform);
 			ChildComponent->DestroyComponent();	
 		}
 	}
+
+	if (AActor* OwnerActor = GetOwner())
+	{
+		constexpr bool bAlsoDestroyComponents = false; // already destroyed!
+		OwnerActor->ClearInstanceComponents(bAlsoDestroyComponents);
+	}
 }
 
 void UText3DComponent::BuildTextMesh(const bool bCleanCache)
@@ -754,7 +764,7 @@ void UText3DComponent::BuildTextMesh(const bool bCleanCache)
 	{
 		if (UText3DComponent* StrongThis = WeakThis.Get())
 		{
-			if(!UE::IsSavingPackage(StrongThis))
+			if (!UE::IsSavingPackage(StrongThis))
 			{
 				StrongThis->BuildTextMeshInternal(bCleanCache);	
 			}
@@ -840,7 +850,7 @@ void UText3DComponent::BuildTextMeshInternal(const bool bCleanCache)
 				continue;
 			}
 
-			if(CharacterMeshes.IsValidIndex(GlyphId))
+			if (CharacterMeshes.IsValidIndex(GlyphId))
 			{
 				UStaticMeshComponent* StaticMeshComponent = CharacterMeshes[GlyphId];
 				StaticMeshComponent->SetStaticMesh(CachedMesh);			
@@ -854,7 +864,7 @@ void UText3DComponent::BuildTextMeshInternal(const bool bCleanCache)
 				UE_LOG(LogText3D, Error, TEXT("CharacterMesh not found at index %d"), GlyphId);
 			}
 
-			if(CharacterKernings.IsValidIndex(GlyphId))
+			if (CharacterKernings.IsValidIndex(GlyphId))
 			{
 				FTransform Transform;
 				Transform.SetLocation(GlyphLocation);
