Text3D material assignment with fastbuild

diff --git a/Engine/Plugins/Experimental/Text3D/Source/Text3D/Private/Glyph.cpp b/Engine/Plugins/Experimental/Text3D/Source/Text3D/Private/Glyph.cpp
--- a/Engine/Plugins/Experimental/Text3D/Source/Text3D/Private/Glyph.cpp
+++ b/Engine/Plugins/Experimental/Text3D/Source/Text3D/Private/Glyph.cpp
@@ -4,7 +4,6 @@
 #include "Materials/Material.h"
 #include "Engine/StaticMesh.h"
 
-
 FText3DGlyph::FText3DGlyph() :
 	StaticMeshAttributes(MeshDescription)
 {
@@ -32,9 +31,7 @@ void FText3DGlyph::Build(UStaticMesh* StaticMesh, UMaterial* DefaultMaterial)
 
 	UStaticMesh::FBuildMeshDescriptionsParams Params;
 	Params.bCommitMeshDescription = true;
-#if !WITH_EDITOR
 	Params.bFastBuild = true;
-#endif
 	StaticMesh->BuildFromMeshDescriptions(MeshDescriptions, Params);
 }
 
diff --git a/Engine/Plugins/Experimental/Text3D/Source/Text3D/Private/MeshCreator.cpp b/Engine/Plugins/Experimental/Text3D/Source/Text3D/Private/MeshCreator.cpp
--- a/Engine/Plugins/Experimental/Text3D/Source/Text3D/Private/MeshCreator.cpp
+++ b/Engine/Plugins/Experimental/Text3D/Source/Text3D/Private/MeshCreator.cpp
@@ -1,13 +1,12 @@
 // Copyright Epic Games, Inc. All Rights Reserved.
 
-
 #include "MeshCreator.h"
-#include "Contour.h"
-#include "Part.h"
-#include "Data.h"
-#include "ContourList.h"
 
 #include "ConstrainedDelaunay2.h"
+#include "Contour.h"
+#include "ContourList.h"
+#include "Data.h"
+#include "Part.h"
 
 using namespace UE::Geometry;
 
@@ -177,6 +176,8 @@ void FMeshCreator::CreateBevelMesh(const float Bevel, const EText3DBevelType Typ
 
 void FMeshCreator::CreateExtrudeMesh(float Extrude, float Bevel, const EText3DBevelType Type, bool bFlipNormals)
 {
+	Bevel = FMath::Max(UE_SMALL_NUMBER, Bevel);
+	
 	if (Type != EText3DBevelType::HalfCircle)
 	{
 		Bevel = FMath::Clamp(Bevel, 0.0f, Extrude / 2.f);
diff --git a/Engine/Plugins/Experimental/Text3D/Source/Text3D/Private/Text3DComponent.cpp b/Engine/Plugins/Experimental/Text3D/Source/Text3D/Private/Text3DComponent.cpp
--- a/Engine/Plugins/Experimental/Text3D/Source/Text3D/Private/Text3DComponent.cpp
+++ b/Engine/Plugins/Experimental/Text3D/Source/Text3D/Private/Text3DComponent.cpp
@@ -909,7 +909,23 @@ float UText3DComponent::MaxBevel() const
 
 void UText3DComponent::UpdateMaterial(const EText3DGroupType Type, UMaterialInterface* Material)
 {
-	const int32 Index = static_cast<int32>(Type);
+	// Material indices are affected by some options
+	const bool bHasBevel = !bOutline && !FMath::IsNearlyZero(Bevel);
+	if (!bHasBevel && Type == EText3DGroupType::Bevel)
+	{
+		return; 
+	}
+	
+	const bool bHasExtrude = !FMath::IsNearlyZero(Extrude);
+	if (!bHasExtrude && Type == EText3DGroupType::Extrude)
+	{
+		return;
+	}
+
+	int32 Index = static_cast<int32>(Type);
+	Index -= !bHasBevel && Type >= EText3DGroupType::Bevel ? 1 : 0; // if no bevel, and the input is bevel or above (bevel, side/extrude, back), offset -1
+	Index -= !bHasExtrude && Type >= EText3DGroupType::Extrude ? 1 : 0; // if no extrude, and the input is side/extrude or above (back), offset -1
+
 	for (UStaticMeshComponent* StaticMeshComponent : CharacterMeshes)
 	{
 		StaticMeshComponent->SetMaterial(Index, Material);
