[Text3D]: Text3D Component removes all instance components from its owning Actor

 1) removed call which was adding StaticMeshes to InstanceComponents
 2) removed calls to remove Components from Owner Actor InstancedComponents, since they shouldn't be there anymore after the change 1)

(regarding point 1: these StaticMeshes are not serialized, so they shouldn't be added to InstanceComponents)

diff --git a/Engine/Plugins/Experimental/Text3D/Source/Text3D/Private/Text3DComponent.cpp b/Engine/Plugins/Experimental/Text3D/Source/Text3D/Private/Text3DComponent.cpp
--- a/Engine/Plugins/Experimental/Text3D/Source/Text3D/Private/Text3DComponent.cpp
+++ b/Engine/Plugins/Experimental/Text3D/Source/Text3D/Private/Text3DComponent.cpp
@@ -294,7 +294,6 @@ bool UText3DComponent::AllocateGlyphs(int32 Num)
 			StaticMeshComponent->SetCastShadow(bCastShadow);
 			CharacterMeshes.Add(StaticMeshComponent);
 
-			GetOwner()->AddInstanceComponent(StaticMeshComponent);
 			StaticMeshComponent->AttachToComponent(CharacterKerningComponent, FAttachmentTransformRules::KeepRelativeTransform);
 		}
 	}
@@ -312,9 +311,7 @@ bool UText3DComponent::AllocateGlyphs(int32 Num)
 				CharacterKerningComponent->DetachFromComponent(FDetachmentTransformRules::KeepRelativeTransform);
 				CharacterKerningComponent->UnregisterComponent();
 				CharacterKerningComponent->DestroyComponent();
-				OwnerActor->RemoveInstanceComponent(CharacterKerningComponent);
 			}
-
 			
 			UStaticMeshComponent* StaticMeshComponent = CharacterMeshes[CharacterIndex];
 			if (IsValid(StaticMeshComponent))
@@ -322,7 +319,6 @@ bool UText3DComponent::AllocateGlyphs(int32 Num)
 				StaticMeshComponent->DetachFromComponent(FDetachmentTransformRules::KeepRelativeTransform);
 				StaticMeshComponent->UnregisterComponent();
 				StaticMeshComponent->DestroyComponent();
-				OwnerActor->RemoveInstanceComponent(StaticMeshComponent);
 			}
 		}
 		
@@ -739,12 +735,6 @@ void UText3DComponent::ClearTextMesh()
 			ChildComponent->DestroyComponent();	
 		}
 	}
-
-	if (AActor* OwnerActor = GetOwner())
-	{
-		constexpr bool bAlsoDestroyComponents = false; // already destroyed!
-		OwnerActor->ClearInstanceComponents(bAlsoDestroyComponents);
-	}
 }
 
 void UText3DComponent::BuildTextMesh(const bool bCleanCache)
