[Text3D] - Text3D Rebuilding meshes when changing relative location

It is actually rebuilding Text meshes OnRegister, even when a rebuild is not necessary.
This affects performance, and also causes flickering e.g when moving owning Actor via details panel, RC or RC controllers.
Issue investigated while working on the jira issue mentioned below.

- Added enum flag to differentiate between different types of actions required by property changes (e.g. mesh update, layout update, nothing)
- Reorganized functions a bit, in order to comply with this change

diff --git a/Engine/Plugins/Experimental/Text3D/Source/Text3D/Private/Text3DComponent.cpp b/Engine/Plugins/Experimental/Text3D/Source/Text3D/Private/Text3DComponent.cpp
--- a/Engine/Plugins/Experimental/Text3D/Source/Text3D/Private/Text3DComponent.cpp
+++ b/Engine/Plugins/Experimental/Text3D/Source/Text3D/Private/Text3DComponent.cpp
@@ -15,6 +15,7 @@
 #include "Engine/StaticMesh.h"
 #include "Materials/Material.h"
 #include "Misc/ScopeExit.h"
+#include "Misc/TransactionObjectEvent.h"
 #include "UObject/ConstructorHelpers.h"
 
 #define LOCTEXT_NAMESPACE "Text3D"
@@ -93,25 +94,63 @@ UText3DComponent::UText3DComponent()
 	MaxHeight = 500.0f;
 	bScaleProportionally = true;
 
-	bPendingBuild = false;
 	bFreezeBuild = false;
+	ModifyFlags = EText3DModifyFlags::All;
 
 	TextScale = FVector::ZeroVector;
 }
 
+void UText3DComponent::BeginDestroy()
+{
+	ClearTextMesh();
+
+	Super::BeginDestroy();
+}
+
+bool UText3DComponent::NeedsMeshRebuild() const
+{
+	return EnumHasAnyFlags(ModifyFlags,EText3DModifyFlags::Geometry);
+}
+
+bool UText3DComponent::NeedsLayoutUpdate() const
+{
+	return EnumHasAnyFlags(ModifyFlags, EText3DModifyFlags::Layout);
+}
+
+void UText3DComponent::MarkForGeometryUpdate()
+{
+	ModifyFlags |= EText3DModifyFlags::Geometry;
+}
+
+void UText3DComponent::MarkForLayoutUpdate()
+{
+	ModifyFlags |= EText3DModifyFlags::Layout;
+}
+
+void UText3DComponent::ClearUpdateFlags()
+{
+	ModifyFlags = EText3DModifyFlags::None;
+}
+
 void UText3DComponent::OnRegister()
 {
 	Super::OnRegister();
-	TextRoot->AttachToComponent(this, FAttachmentTransformRules::KeepRelativeTransform);
 
-	// Forces rebuild (if build is in progress)
-	bIsBuilding = false;
-	BuildTextMesh();
+	if (!TextRoot->IsAttachedTo(this))
+	{
+		TextRoot->AttachToComponent(this, FAttachmentTransformRules::KeepRelativeTransform);
+	}
+
+	Update();
 }
 
 void UText3DComponent::OnUnregister()
 {
-	ClearTextMesh();
+	if (IsBeingDestroyed())
+	{
+		ClearTextMesh();
+	}
+
 	Super::OnUnregister();
 }
 
@@ -148,12 +187,49 @@ void UText3DComponent::PostEditChangeProperty(FPropertyChangedEvent& PropertyCha
 			break;
 		}
 		}
+
+		MarkForGeometryUpdate();
 	}
 	else if (Name == BevelSegmentsPropertyName)
 	{
 		// Force minimum bevel segments based on the BevelType
 		SetBevelSegments(BevelSegments);
 	}
+	else if (Name == GET_MEMBER_NAME_CHECKED(UText3DComponent, Font) ||
+			 Name == GET_MEMBER_NAME_CHECKED(UText3DComponent, Text) ||
+			 Name == GET_MEMBER_NAME_CHECKED(UText3DComponent, OutlineExpand) ||
+			 Name == GET_MEMBER_NAME_CHECKED(UText3DComponent, bOutline) ||
+			 Name == GET_MEMBER_NAME_CHECKED(UText3DComponent, Extrude) ||
+			 Name == GET_MEMBER_NAME_CHECKED(UText3DComponent, Bevel))
+	{
+		MarkForGeometryUpdate();
+	}
+	else if (Name == GET_MEMBER_NAME_CHECKED(UText3DComponent, HorizontalAlignment) ||
+			 Name == GET_MEMBER_NAME_CHECKED(UText3DComponent, VerticalAlignment) ||
+			 Name == GET_MEMBER_NAME_CHECKED(UText3DComponent, Kerning) ||
+			 Name == GET_MEMBER_NAME_CHECKED(UText3DComponent, LineSpacing) ||
+			 Name == GET_MEMBER_NAME_CHECKED(UText3DComponent, WordSpacing) ||
+			 Name == GET_MEMBER_NAME_CHECKED(UText3DComponent, bHasMaxWidth) ||
+			 Name == GET_MEMBER_NAME_CHECKED(UText3DComponent, MaxWidth) ||
+			 Name == GET_MEMBER_NAME_CHECKED(UText3DComponent, bHasMaxHeight) ||
+			 Name == GET_MEMBER_NAME_CHECKED(UText3DComponent, MaxHeight) ||
+			 Name == GET_MEMBER_NAME_CHECKED(UText3DComponent, bScaleProportionally))
+	{
+		MarkForLayoutUpdate();
+	}
+
+	Update();
+}
+
+void UText3DComponent::PostTransacted(const FTransactionObjectEvent& TransactionEvent)
+{
+	Super::PostTransacted(TransactionEvent);
+
+	if (TransactionEvent.GetEventType() == ETransactionObjectEventType::UndoRedo)
+	{
+		ModifyFlags |= EText3DModifyFlags::All;
+		Update();
+	}
 }
 #endif
 
@@ -162,7 +238,8 @@ void UText3DComponent::SetText(const FText& Value)
 	if (!Text.EqualTo(Value))
 	{
 		Text = Value;
-		Rebuild();
+		MarkForGeometryUpdate();
+		Update();
 	}
 }
 
@@ -171,7 +248,8 @@ void UText3DComponent::SetFont(UFont* const InFont)
 	if (Font != InFont)
 	{
 		Font = InFont;
-		Rebuild();
+		MarkForGeometryUpdate();
+		Update();
 	}
 }
 
@@ -179,8 +257,9 @@ void UText3DComponent::SetOutline(const bool bValue)
 {
 	if (bOutline != bValue)
 	{
-		bOutline = bValue;		
-		Rebuild();
+		bOutline = bValue;
+		MarkForGeometryUpdate();
+		Update();
 	}
 }
 
@@ -190,7 +269,8 @@ void UText3DComponent::SetOutlineExpand(const float Value)
 	if (!FMath::IsNearlyEqual(OutlineExpand, NewValue))
 	{
 		OutlineExpand = NewValue;
-		Rebuild();
+		MarkForGeometryUpdate();
+		Update();
 	}
 }
 
@@ -200,8 +280,9 @@ void UText3DComponent::SetExtrude(const float Value)
 	if (!FMath::IsNearlyEqual(Extrude, NewValue))
 	{
 		Extrude = NewValue;
+		MarkForGeometryUpdate();
 		CheckBevel();
-		Rebuild();
+		Update();
 	}
 }
 
@@ -212,7 +293,8 @@ void UText3DComponent::SetBevel(const float Value)
 	if (!FMath::IsNearlyEqual(Bevel, NewValue))
 	{
 		Bevel = NewValue;
-		Rebuild();
+		MarkForGeometryUpdate();
+		Update();
 	}
 }
 
@@ -221,7 +303,8 @@ void UText3DComponent::SetBevelType(const EText3DBevelType Value)
 	if (BevelType != Value)
 	{
 		BevelType = Value;
-		Rebuild();
+		MarkForGeometryUpdate();
+		Update();
 	}
 }
 
@@ -237,7 +320,8 @@ void UText3DComponent::SetBevelSegments(const int32 Value)
 	if (BevelSegments != NewValue)
 	{
 		BevelSegments = NewValue;
-		Rebuild();
+		MarkForGeometryUpdate();
+		Update();
 	}
 }
 
@@ -365,7 +449,7 @@ void UText3DComponent::SetMaterial(const EText3DGroupType Type, UMaterialInterfa
 	UMaterialInterface* OldMaterial = GetMaterial(Type);
 	if (Value != OldMaterial)
 	{
-		switch(Type)
+		switch (Type)
 		{
 		case EText3DGroupType::Front:
 		{
@@ -498,9 +582,9 @@ void UText3DComponent::SetFreeze(const bool bFreeze)
 	bFreezeBuild = bFreeze;
 	if (bFreeze)
 	{
-		bPendingBuild = false;
+		ModifyFlags = EText3DModifyFlags::Unfreeze;
 	}
-	else if (bPendingBuild)
+	else if (EnumHasAnyFlags(ModifyFlags,EText3DModifyFlags::Unfreeze))
 	{
 		Rebuild();
 	}
@@ -558,13 +642,24 @@ const TArray<UStaticMeshComponent*>& UText3DComponent::GetGlyphMeshComponents()
 
 void UText3DComponent::Rebuild(const bool bCleanCache)
 {
-	bPendingBuild = true;
 	if (!bFreezeBuild)
 	{
 		BuildTextMesh(bCleanCache);
 	}
 }
 
+void UText3DComponent::Update()
+{
+	if (NeedsMeshRebuild())
+	{
+		Rebuild();
+	}
+	else if (NeedsLayoutUpdate())
+	{
+		UpdateTransforms();
+	}
+}
+
 void UText3DComponent::CalculateTextWidth()
 {
 	for (FShapedGlyphLine& ShapedLine : ShapedText->Lines)
@@ -696,6 +791,8 @@ void UText3DComponent::UpdateTransforms()
 			GlyphIndex++;
 		}
 	}
+
+	ModifyFlags &= EText3DModifyFlags::Layout; 
 }
 
 void UText3DComponent::ClearTextMesh()
@@ -723,16 +820,19 @@ void UText3DComponent::ClearTextMesh()
 	}
 	CharacterKernings.Reset();
 
-	constexpr bool bIncludeChildDescendants = true;
-	TArray<USceneComponent*> ChildComponents;
-	TextRoot->GetChildrenComponents(bIncludeChildDescendants, ChildComponents);
-
-	for (USceneComponent* ChildComponent : ChildComponents)
+	if (TextRoot)
 	{
-		if (IsValid(ChildComponent))
+		TArray<USceneComponent*> ChildComponents;
+		constexpr bool bIncludeChildDescendants = true;
+		TextRoot->GetChildrenComponents(bIncludeChildDescendants, ChildComponents);
+
+		for (USceneComponent* ChildComponent : ChildComponents)
 		{
-			ChildComponent->DetachFromComponent(FDetachmentTransformRules::KeepRelativeTransform);
-			ChildComponent->DestroyComponent();	
+			if (IsValid(ChildComponent))
+			{
+				ChildComponent->DetachFromComponent(FDetachmentTransformRules::KeepRelativeTransform);
+				ChildComponent->DestroyComponent();	
+			}
 		}
 	}
 }
@@ -773,7 +873,6 @@ void UText3DComponent::BuildTextMeshInternal(const bool bCleanCache)
 		return;
 	}
 
-	bPendingBuild = false;
 	CheckBevel();
 
 	ClearTextMesh();
@@ -877,6 +976,8 @@ void UText3DComponent::BuildTextMeshInternal(const bool bCleanCache)
 
 	TextGeneratedNativeDelegate.Broadcast();
 	TextGeneratedDelegate.Broadcast();
+	
+	ClearUpdateFlags();
 
 	if (bCleanCache)
 	{
diff --git a/Engine/Plugins/Experimental/Text3D/Source/Text3D/Public/Text3DComponent.h b/Engine/Plugins/Experimental/Text3D/Source/Text3D/Public/Text3DComponent.h
--- a/Engine/Plugins/Experimental/Text3D/Source/Text3D/Public/Text3DComponent.h
+++ b/Engine/Plugins/Experimental/Text3D/Source/Text3D/Public/Text3DComponent.h
@@ -27,6 +27,17 @@ enum class EText3DHorizontalTextAlignment : uint8
 	Right			UMETA(DisplayName = "Right"),
 };
 
+UENUM()
+enum class EText3DModifyFlags : uint8
+{
+	None = 0,
+	Layout = 1 << 0,
+	Geometry = 1 << 1,
+	Unfreeze = 1 << 2,
+    
+	All = Layout | Geometry | Unfreeze
+};
+ENUM_CLASS_FLAGS(EText3DModifyFlags)
 
 UCLASS(ClassGroup = (Text3D), meta=(BlueprintSpawnableComponent))
 class TEXT3D_API UText3DComponent : public USceneComponent
@@ -36,11 +47,13 @@ class TEXT3D_API UText3DComponent : public USceneComponent
 public:
 	UText3DComponent();
 
+	virtual void BeginDestroy() override;
 	virtual void OnRegister() override;
 	virtual void OnUnregister() override;
 
 #if WITH_EDITOR
 	virtual void PostEditChangeProperty(struct FPropertyChangedEvent& PropertyChangedEvent) override;
+	virtual void PostTransacted(const FTransactionObjectEvent& TransactionEvent) override;
 #endif
 
 	/** The text to generate a 3d mesh */
@@ -275,7 +288,7 @@ protected:
 	virtual void OnHiddenInGameChanged() override;
 
 private:
-	UPROPERTY(Transient)
+	UPROPERTY()
 	TObjectPtr<class USceneComponent> TextRoot;
 
 	UPROPERTY(BlueprintAssignable, Category = Events, meta = (AllowPrivateAccess = true, DisplayName = "On Text Generated"))
@@ -286,8 +299,8 @@ private:
 	/** Flagged as true when the text mesh is being built. */
 	std::atomic<bool> bIsBuilding;
 	
-	bool bPendingBuild;
 	bool bFreezeBuild;
+	EText3DModifyFlags ModifyFlags;
 
 	FVector TextScale;
 
@@ -308,6 +321,7 @@ private:
 
 	void UpdateTransforms();
 	void Rebuild(const bool bCleanCache = false);
+	void Update();
 	void ClearTextMesh();
 	void BuildTextMesh(const bool bCleanCache = false);
 	void BuildTextMeshInternal(const bool bCleanCache);
@@ -320,4 +334,10 @@ private:
 	float GetTextHeight() const;
 	void CalculateTextScale();
 	FVector GetLineLocation(int32 LineIndex);
+	
+	bool NeedsMeshRebuild() const;
+	bool NeedsLayoutUpdate() const;
+	void MarkForGeometryUpdate();
+	void MarkForLayoutUpdate();
+	void ClearUpdateFlags();
 };
