Text3D fallback font support

diff --git a/Engine/Plugins/Experimental/Text3D/Source/Text3D/Private/Text3DComponent.cpp b/Engine/Plugins/Experimental/Text3D/Source/Text3D/Private/Text3DComponent.cpp
--- a/Engine/Plugins/Experimental/Text3D/Source/Text3D/Private/Text3DComponent.cpp
+++ b/Engine/Plugins/Experimental/Text3D/Source/Text3D/Private/Text3DComponent.cpp
@@ -2,20 +2,23 @@
 
 #include "Text3DComponent.h"
 
-#include "Glyph.h"
-#include "MeshCreator.h"
-#include "Text3DEngineSubsystem.h"
-#include "Text3DPrivate.h"
-#include "TextShaper.h"
 #include "Algo/Count.h"
 #include "Async/Async.h"
 #include "Components/StaticMeshComponent.h"
 #include "Engine/Engine.h"
 #include "Engine/Font.h"
 #include "Engine/StaticMesh.h"
+#include "Fonts/SlateTextShaper.h"
+#include "Framework/Text/PlainTextLayoutMarshaller.h"
+#include "Glyph.h"
 #include "Materials/Material.h"
+#include "MeshCreator.h"
 #include "Misc/ScopeExit.h"
 #include "Misc/TransactionObjectEvent.h"
+#include "Styling/StyleDefaults.h"
+#include "Text3DEngineSubsystem.h"
+#include "Text3DPrivate.h"
+#include "TextShaper.h"
 #include "UObject/ConstructorHelpers.h"
 
 #define LOCTEXT_NAMESPACE "Text3D"
@@ -141,7 +144,7 @@ void UText3DComponent::OnRegister()
 		TextRoot->AttachToComponent(this, FAttachmentTransformRules::KeepRelativeTransform);
 	}
 
-	Update();
+	RebuildInternal();
 }
 
 void UText3DComponent::OnUnregister()
@@ -218,7 +221,7 @@ void UText3DComponent::PostEditChangeProperty(FPropertyChangedEvent& PropertyCha
 		MarkForLayoutUpdate();
 	}
 
-	Update();
+	RebuildInternal();
 }
 
 void UText3DComponent::PostTransacted(const FTransactionObjectEvent& TransactionEvent)
@@ -228,18 +231,26 @@ void UText3DComponent::PostTransacted(const FTransactionObjectEvent& Transaction
 	if (TransactionEvent.GetEventType() == ETransactionObjectEventType::UndoRedo)
 	{
 		ModifyFlags |= EText3DModifyFlags::All;
-		Update();
+		RebuildInternal();
 	}
 }
 #endif
 
+void UText3DComponent::SetRefreshOnChange(const bool& Value)
+{
+	if (bRefreshOnChange != Value)
+	{
+		bRefreshOnChange = Value;
+	}
+}
+
 void UText3DComponent::SetText(const FText& Value)
 {
 	if (!Text.EqualTo(Value))
 	{
 		Text = Value;
 		MarkForGeometryUpdate();
-		Update();
+		RebuildInternal();
 	}
 }
 
@@ -249,7 +260,7 @@ void UText3DComponent::SetFont(UFont* const InFont)
 	{
 		Font = InFont;
 		MarkForGeometryUpdate();
-		Update();
+		RebuildInternal();
 	}
 }
 
@@ -258,8 +269,8 @@ void UText3DComponent::SetOutline(const bool bValue)
 	if (bOutline != bValue)
 	{
 		bOutline = bValue;
-		MarkForGeometryUpdate();
-		Update();
+		MarkForGeometryUpdate();	
+		RebuildInternal();
 	}
 }
 
@@ -270,7 +281,7 @@ void UText3DComponent::SetOutlineExpand(const float Value)
 	{
 		OutlineExpand = NewValue;
 		MarkForGeometryUpdate();
-		Update();
+		RebuildInternal();
 	}
 }
 
@@ -282,7 +293,7 @@ void UText3DComponent::SetExtrude(const float Value)
 		Extrude = NewValue;
 		MarkForGeometryUpdate();
 		CheckBevel();
-		Update();
+		RebuildInternal();
 	}
 }
 
@@ -294,7 +305,7 @@ void UText3DComponent::SetBevel(const float Value)
 	{
 		Bevel = NewValue;
 		MarkForGeometryUpdate();
-		Update();
+		RebuildInternal();
 	}
 }
 
@@ -304,7 +315,7 @@ void UText3DComponent::SetBevelType(const EText3DBevelType Value)
 	{
 		BevelType = Value;
 		MarkForGeometryUpdate();
-		Update();
+		RebuildInternal();
 	}
 }
 
@@ -321,7 +332,7 @@ void UText3DComponent::SetBevelSegments(const int32 Value)
 	{
 		BevelSegments = NewValue;
 		MarkForGeometryUpdate();
-		Update();
+		RebuildInternal();
 	}
 }
 
@@ -586,7 +597,7 @@ void UText3DComponent::SetFreeze(const bool bFreeze)
 	}
 	else if (EnumHasAnyFlags(ModifyFlags,EText3DModifyFlags::Unfreeze))
 	{
-		Rebuild();
+		RebuildInternal();
 	}
 }
 
@@ -640,19 +651,23 @@ const TArray<UStaticMeshComponent*>& UText3DComponent::GetGlyphMeshComponents()
 	return CharacterMeshes;
 }
 
-void UText3DComponent::Rebuild(const bool bCleanCache)
+void UText3DComponent::Rebuild()
 {
-	if (!bFreezeBuild)
-	{
-		BuildTextMesh(bCleanCache);
-	}
+	// Rebuild, ignoring RefreshOnChange
+	RebuildInternal(false);
 }
 
-void UText3DComponent::Update()
+void UText3DComponent::RebuildInternal(const bool& bIsAutoUpdate, const bool& bCleanCache)
 {
-	if (NeedsMeshRebuild())
+	// If this is an auto update, but the flag is off, ignore this rebuild request
+	if (bIsAutoUpdate && bRefreshOnChange == false)
 	{
-		Rebuild();
+		return;
+	}
+
+	if (NeedsMeshRebuild() && !bFreezeBuild)
+	{
+		BuildTextMesh(bCleanCache);
 	}
 	else if (NeedsLayoutUpdate())
 	{
@@ -776,7 +791,7 @@ void UText3DComponent::UpdateTransforms()
 		for (int32 LineGlyph = 0; LineGlyph < Line.GlyphsToRender.Num(); LineGlyph++)
 		{
 			const FVector CharLocation = Location;
-			Location.Y += Line.GetAdvanced(LineGlyph, Kerning, WordSpacing);
+			Location.Y += Line.GetAdvance(LineGlyph, Kerning, WordSpacing);
 			if (!Line.GlyphsToRender[LineGlyph].bIsVisible)
 			{
 				continue;
@@ -837,7 +852,7 @@ void UText3DComponent::ClearTextMesh()
 	}
 }
 
-void UText3DComponent::BuildTextMesh(const bool bCleanCache)
+void UText3DComponent::BuildTextMesh(const bool& bCleanCache)
 {
 	// If we're already building, or have a build pending, don't do anything.
 	if (bIsBuilding)
@@ -862,7 +877,7 @@ void UText3DComponent::BuildTextMesh(const bool bCleanCache)
 	});
 }
 
-void UText3DComponent::BuildTextMeshInternal(const bool bCleanCache)
+void UText3DComponent::BuildTextMeshInternal(const bool& bCleanCache)
 {
 	TRACE_CPUPROFILER_EVENT_SCOPE(TEXT("UText3DComponent::Rebuild"));
 
@@ -897,7 +912,54 @@ void UText3DComponent::BuildTextMeshInternal(const bool bCleanCache)
 	ShapedText->LineHeight = Face->size->metrics.height * FontInverseScale;
 	ShapedText->FontAscender = Face->size->metrics.ascender * FontInverseScale;
 	ShapedText->FontDescender = Face->size->metrics.descender * FontInverseScale;
-	FTextShaper::Get()->ShapeBidirectionalText(Face, Text.ToString(), ShapedText->Lines);
+
+	constexpr int32 AdjustedFontSize = 48; // Magic number that makes font scale consistent with previous implementation
+	FSlateFontInfo FontInfo(Font, AdjustedFontSize);
+	FontInfo.CompositeFont = FStyleDefaults::GetFontInfo().CompositeFont;
+	
+	FTextBlockStyle Style;
+	Style.SetFont(FontInfo);
+
+	if (!TextLayout.IsValid())
+	{
+		TextLayout = MakeShared<FText3DLayout>(Style);
+	}
+
+	if (!TextLayoutMarshaller.IsValid())
+	{
+		TextLayoutMarshaller = FPlainTextLayoutMarshaller::Create();
+	}
+
+	FTextShaper::Get()->ShapeBidirectionalText(Style, Text.ToString(), TextLayout, TextLayoutMarshaller, ShapedText->Lines);
+
+	const int32 ApproximateGlyphNum = Algo::TransformAccumulate(ShapedText->Lines,
+		[](const FShapedGlyphLine& InLine)
+		{
+			return InLine.GlyphsToRender.Num();
+		},
+		0);
+
+	TMap<uint32, TSharedPtr<FFreeTypeFace>> GlyphIndexToFontFace;
+	GlyphIndexToFontFace.Reserve(ApproximateGlyphNum);
+
+	for (int32 LineIndex = 0; LineIndex < ShapedText->Lines.Num(); ++LineIndex)
+	{
+		for (const FShapedGlyphEntry& GlyphEntry : ShapedText->Lines[LineIndex].GlyphsToRender)
+		{
+			if (!GlyphEntry.FontFaceData.IsValid())
+			{
+				// Add as nullptr if not already in map
+				if (!GlyphIndexToFontFace.Contains(GlyphEntry.GlyphIndex))
+				{
+					GlyphIndexToFontFace.Add(GlyphEntry.GlyphIndex, nullptr);
+				}
+				
+				continue;
+			}
+			
+			GlyphIndexToFontFace.FindOrAdd(GlyphEntry.GlyphIndex, GlyphEntry.FontFaceData->FontFace.Pin());
+		}
+	}
 	
 	CalculateTextWidth();
 	CalculateTextScale();
@@ -922,7 +984,7 @@ void UText3DComponent::BuildTextMeshInternal(const bool bCleanCache)
 		for (int32 LineGlyph = 0; LineGlyph < ShapedLine.GlyphsToRender.Num(); LineGlyph++)
 		{
 			FVector GlyphLocation = Location;
-			Location.Y += ShapedLine.GetAdvanced(LineGlyph, Kerning, WordSpacing);
+			Location.Y += ShapedLine.GetAdvance(LineGlyph, Kerning, WordSpacing);
 
 			const FShapedGlyphEntry& ShapedGlyph = ShapedLine.GlyphsToRender[LineGlyph];
 			if (!ShapedGlyph.bIsVisible)
@@ -933,7 +995,8 @@ void UText3DComponent::BuildTextMeshInternal(const bool bCleanCache)
 			// Count even when mesh is nullptr (allocation still creates components to avoid mesh building in allocation step)
 			const int32 GlyphId = GlyphIndex++;
 
-			UStaticMesh* CachedMesh = CachedFontData.GetGlyphMesh(ShapedGlyph.GlyphIndex, FGlyphMeshParameters{Extrude, Bevel, BevelType, BevelSegments, bOutline, OutlineExpand});
+			TSharedPtr<FFreeTypeFace> FontFace = GlyphIndexToFontFace[ShapedGlyph.GlyphIndex];
+			UStaticMesh* CachedMesh = CachedFontData.GetGlyphMesh(ShapedGlyph.GlyphIndex, FGlyphMeshParameters{Extrude, Bevel, BevelType, BevelSegments, bOutline, OutlineExpand}, FontFace);
 			if (!CachedMesh || FMath::IsNearlyZero(CachedMesh->GetBounds().SphereRadius))
 			{
 				continue;
diff --git a/Engine/Plugins/Experimental/Text3D/Source/Text3D/Private/Text3DEngineSubsystem.cpp b/Engine/Plugins/Experimental/Text3D/Source/Text3D/Private/Text3DEngineSubsystem.cpp
--- a/Engine/Plugins/Experimental/Text3D/Source/Text3D/Private/Text3DEngineSubsystem.cpp
+++ b/Engine/Plugins/Experimental/Text3D/Source/Text3D/Private/Text3DEngineSubsystem.cpp
@@ -10,9 +10,9 @@
 #include "MeshCreator.h"
 #include "GlyphLoader.h"
 #include "ContourNode.h"
+#include "Fonts/FontCacheFreeType.h"
 #include "UObject/ConstructorHelpers.h"
 
-
 FCachedFontMeshes::FCachedFontMeshes()
 {
 	CacheCounter = MakeShared<int32>(0);
@@ -195,7 +195,7 @@ TSharedPtr<int32> FCachedFontData::GetMeshesCacheCounter(const FGlyphMeshParamet
 	return CachedMeshes.GetCacheCounter();
 }
 
-UStaticMesh* FCachedFontData::GetGlyphMesh(uint32 GlyphIndex, const FGlyphMeshParameters& Parameters)
+UStaticMesh* FCachedFontData::GetGlyphMesh(uint32 GlyphIndex, const FGlyphMeshParameters& Parameters, const TSharedPtr<FFreeTypeFace>& FontFaceData)
 {
 	const uint32 HashParameters = GetTypeHash(Parameters);
 	FCachedFontMeshes& CachedMeshes = Meshes.FindOrAdd(HashParameters);
@@ -211,7 +211,7 @@ UStaticMesh* FCachedFontData::GetGlyphMesh(uint32 GlyphIndex, const FGlyphMeshPa
 	HashGroup = HashCombine(HashGroup, GetTypeHash(GlyphIndex));
 	const FString StaticMeshName = FString::Printf(TEXT("Text3D_Char_%u_%u"), HashGroup, HashParameters);
 
-	const TSharedContourNode Root = GetGlyphContours(GlyphIndex);
+	const TSharedContourNode Root = GetGlyphContours(GlyphIndex, FontFaceData);
 	if (Root->Children.Num() == 0)
 	{
 		return nullptr;
@@ -241,16 +241,30 @@ const FString& FCachedFontData::GetFontName()
 	return FontName;
 }
 
-TSharedContourNode FCachedFontData::GetGlyphContours(uint32 GlyphIndex)
+TSharedContourNode FCachedFontData::GetGlyphContours(uint32 GlyphIndex, const TSharedPtr<FFreeTypeFace>& FontFaceData)
 {
-	check(FreeTypeFace);
+	const FT_Face& GlyphFreeTypeFace = FontFaceData.IsValid() && FontFaceData->IsFaceValid()
+		? FontFaceData->GetFace()
+		: FreeTypeFace;
+	check(GlyphFreeTypeFace);
 
-	if (FT_Load_Glyph(FreeTypeFace, GlyphIndex, FT_LOAD_NO_HINTING | FT_LOAD_NO_BITMAP))
+	if (FT_Load_Glyph(GlyphFreeTypeFace, GlyphIndex, FT_LOAD_NO_HINTING | FT_LOAD_NO_BITMAP))
 	{
 		return nullptr;
 	}
 
-	FGlyphLoader GlyphLoader(FreeTypeFace->glyph);
+#if UE_BUILD_DEBUG
+	// If true, the symbol itself was probably bitmap-only, and this will have placeholder contours
+	bool bIsLastResort = FontFaceData->GetFace()->family_name
+		? FString(FontFaceData->GetFace()->family_name) == TEXT("LastResort")
+		: false;
+	if (bIsLastResort)
+	{
+		UE_LOG(LogText3D, Warning, TEXT("The following glyph index is bitmap-only: %i"), GlyphIndex);
+	}
+#endif
+
+	FGlyphLoader GlyphLoader(GlyphFreeTypeFace->glyph);
 	TSharedContourNode Root = GlyphLoader.GetContourList();
 
 	return Root;
diff --git a/Engine/Plugins/Experimental/Text3D/Source/Text3D/Private/Text3DEngineSubsystem.h b/Engine/Plugins/Experimental/Text3D/Source/Text3D/Private/Text3DEngineSubsystem.h
--- a/Engine/Plugins/Experimental/Text3D/Source/Text3D/Private/Text3DEngineSubsystem.h
+++ b/Engine/Plugins/Experimental/Text3D/Source/Text3D/Private/Text3DEngineSubsystem.h
@@ -11,6 +11,8 @@
 
 #include "Text3DEngineSubsystem.generated.h"
 
+class FFreeTypeFace;
+
 USTRUCT()
 struct FGlyphMeshParameters
 {
@@ -87,8 +89,8 @@ public:
 	TSharedPtr<int32> GetCacheCounter();
 	TSharedPtr<int32> GetMeshesCacheCounter(const FGlyphMeshParameters& Parameters);
 
-	UStaticMesh* GetGlyphMesh(uint32 GlyphIndex, const FGlyphMeshParameters& Parameters);
-	TSharedContourNode GetGlyphContours(uint32 GlyphIndex);
+	UStaticMesh* GetGlyphMesh(uint32 GlyphIndex, const FGlyphMeshParameters& Parameters, const TSharedPtr<FFreeTypeFace>& FontFaceData = nullptr);
+	TSharedContourNode GetGlyphContours(uint32 GlyphIndex, const TSharedPtr<FFreeTypeFace>& FontFaceData = nullptr);
 
 	UPROPERTY()
 	TObjectPtr<class UFont> Font;
diff --git a/Engine/Plugins/Experimental/Text3D/Source/Text3D/Private/Text3DPrivate.cpp b/Engine/Plugins/Experimental/Text3D/Source/Text3D/Private/Text3DPrivate.cpp
--- a/Engine/Plugins/Experimental/Text3D/Source/Text3D/Private/Text3DPrivate.cpp
+++ b/Engine/Plugins/Experimental/Text3D/Source/Text3D/Private/Text3DPrivate.cpp
@@ -18,15 +18,12 @@ FText3DModule::FText3DModule()
 void FText3DModule::StartupModule()
 {
 	FT_Init_FreeType(&FreeTypeLib);
-	FTextShaper::Initialize();
 }
 
 void FText3DModule::ShutdownModule()
 {
 	FT_Done_FreeType(FreeTypeLib);
 	FreeTypeLib = nullptr;
-
-	FTextShaper::Cleanup();
 }
 
 FT_Library FText3DModule::GetFreeTypeLibrary()
diff --git a/Engine/Plugins/Experimental/Text3D/Source/Text3D/Private/TextShaper.cpp b/Engine/Plugins/Experimental/Text3D/Source/Text3D/Private/TextShaper.cpp
--- a/Engine/Plugins/Experimental/Text3D/Source/Text3D/Private/TextShaper.cpp
+++ b/Engine/Plugins/Experimental/Text3D/Source/Text3D/Private/TextShaper.cpp
@@ -1,314 +1,77 @@
 // Copyright Epic Games, Inc. All Rights Reserved.
 
-
 #include "TextShaper.h"
+
+#include "Brushes/SlateNoResource.h"
 #include "Containers/Array.h"
 #include "Containers/UnrealString.h"
 #include "Fonts/FontCache.h"
+#include "Framework/Text/PlainTextLayoutMarshaller.h"
+#include "Framework/Text/ShapedTextCache.h"
+#include "Framework/Text/SlateTextRun.h"
 #include "Internationalization/Text.h"
 
-THIRD_PARTY_INCLUDES_START
-#include "hb.h"
-#include "hb-ft.h"
-
-#if WITH_EDITOR
-extern "C"
-{
-void* HarfBuzzMalloc(size_t InSizeBytes)
+FText3DLayout::FText3DLayout(const FTextBlockStyle& InStyle): TextStyle(InStyle)
 {
-	LLM_SCOPE(ELLMTag::UI);
-	return FMemory::Malloc(InSizeBytes);
+	static const FSlateBrush EmptyBrush = FSlateNoResource{};
+	TextStyle.SetUnderlineBrush(EmptyBrush);
+	TextStyle.SetStrikeBrush(EmptyBrush);
 }
 
-void* HarfBuzzCalloc(size_t InNumItems, size_t InItemSizeBytes)
+TSharedRef<IRun> FText3DLayout::CreateDefaultTextRun(
+	const TSharedRef<FString>& NewText,
+	const FTextRange& NewRange) const
 {
-	LLM_SCOPE(ELLMTag::UI);
-	const size_t AllocSizeBytes = InNumItems * InItemSizeBytes;
-	if (AllocSizeBytes > 0)
-	{
-		void* Ptr = FMemory::Malloc(AllocSizeBytes);
-		FMemory::Memzero(Ptr, AllocSizeBytes);
-		return Ptr;
-	}
-	return nullptr;
+	return FSlateTextRun::Create(FRunInfo(), NewText, TextStyle, NewRange);
 }
 
-void* HarfBuzzRealloc(void* InPtr, size_t InSizeBytes)
+TSharedPtr<FTextShaper> FTextShaper::Get()
 {
-	LLM_SCOPE(ELLMTag::UI);
-	return FMemory::Realloc(InPtr, InSizeBytes);
+	static TSharedPtr<FTextShaper> Instance = MakeShared<FTextShaper>(FPrivateToken{});
+	return Instance;
 }
 
-void HarfBuzzFree(void* InPtr)
+void FTextShaper::ShapeBidirectionalText(
+	const FTextBlockStyle& Style,
+	const FString& Text,
+	const TSharedPtr<FTextLayout>& TextLayout,
+	const TSharedPtr<ITextLayoutMarshaller>& TextMarshaller,
+	TArray<FShapedGlyphLine>& OutShapedLines)
 {
-	FMemory::Free(InPtr);
-}
-
-} // extern "C"
-#endif // WITH_EDITOR
-THIRD_PARTY_INCLUDES_END
-
-FTextShaper* FTextShaper::Instance = nullptr;
-
-void FTextShaper::Initialize()
-{
-	if (Instance)
-		return;
-
-	Instance = new FTextShaper();
-}
-
-void FTextShaper::Cleanup()
-{
-	delete Instance;
-	Instance = nullptr;
-}
-
-FTextShaper::FTextShaper()
-{
-	TextBiDiDetection = TextBiDi::CreateTextBiDi();
-}
-
-void FTextShaper::ShapeBidirectionalText(const FT_Face Face, const FString& Text, TArray<FShapedGlyphLine>& OutShapedLines)
-{
-	TextBiDi::ETextDirection Direction = TextBiDiDetection->ComputeBaseDirection(Text);
-	check(Direction != TextBiDi::ETextDirection::Mixed);
-	check(Face);
-
-	OutShapedLines.AddDefaulted();
-
-	TArray<TextBiDi::FTextDirectionInfo> TextDirectionInfos;
-	TextBiDiDetection->ComputeTextDirection(Text, Direction, TextDirectionInfos);
-	for (const TextBiDi::FTextDirectionInfo& TextDirectionInfo : TextDirectionInfos)
-	{
-		if (TextDirectionInfo.Length == 0)
-		{
-			continue;
-		}
+	TextLayout->ClearLines();
+	TextLayout->ClearLineHighlights();
+	TextLayout->ClearRunRenderers();
+	
+	TextMarshaller->SetText(Text, *TextLayout);
+	TextMarshaller->ClearDirty();
 
-		if (TextDirectionInfo.TextDirection == TextBiDi::ETextDirection::RightToLeft)
-		{
-			PerformHarfBuzzTextShaping(Face, *Text, TextDirectionInfo.StartIndex, TextDirectionInfo.StartIndex + TextDirectionInfo.Length, OutShapedLines);
-		}
-		else
-		{
-			PerformKerningTextShaping(Face, *Text, TextDirectionInfo.StartIndex, TextDirectionInfo.StartIndex + TextDirectionInfo.Length, OutShapedLines);
-		}
-	}
+	TextLayout->UpdateLayout();
 
-	if (Direction == TextBiDi::ETextDirection::RightToLeft)
-	{
-		Algo::Reverse(OutShapedLines);
-	}
-}
+	TArray<FTextLayout::FLineView> LineViews = TextLayout->GetLineViews();
 
-void FTextShaper::PerformKerningTextShaping(const FT_Face Face, const TCHAR* Text, const int32 StartIndex, const int32 EndIndex, TArray<FShapedGlyphLine>& OutShapedLines)
-{
-	const bool bHasKerning = FT_HAS_KERNING(Face) != 0;
-	for (int32 Index = StartIndex; Index < EndIndex; Index++)
+	// @note: mimics FSlateTextLayout::OnPaint
+	for (const FTextLayout::FLineView& Line : LineViews)
 	{
-		if (InsertSubstituteGlyphs(Face, Text, Index, OutShapedLines))
+		FShapedGlyphLine& ShapedLine = OutShapedLines.AddDefaulted_GetRef();
+		for (const TSharedRef<ILayoutBlock>& Block : Line.Blocks)
 		{
-			continue;
-		}
-
-		const TCHAR CurrentChar = Text[Index];
-		const bool bIsZeroWidthSpace = CurrentChar == TEXT('\u200B');
-		bool bIsWhitespace = bIsZeroWidthSpace || FText::IsWhitespace(CurrentChar);
-
-		FT_Load_Char(Face, CurrentChar, FT_LOAD_DEFAULT);
-		uint32 GlyphIndex = FT_Get_Char_Index(Face, CurrentChar);
-		if (GlyphIndex == 0)	// Get Space instead of invalid character
-		{
-			GlyphIndex = FT_Get_Char_Index(Face, ' ');
-			bIsWhitespace = true;
-		}
-		
-
-		int16 XAdvance = 0;
-		if (!bIsZeroWidthSpace)
-		{
-			FT_Fixed AdvanceData = 0;
-			if (FT_Get_Advance(Face, GlyphIndex, 0, &AdvanceData) == 0)
-			{
-				XAdvance = ((AdvanceData + (1 << 9)) >> 10) * FontInverseScale;
-			}
-		}
-
-		const int32 CurrentGlyphEntryIndex = OutShapedLines.Last().GlyphsToRender.AddDefaulted();
-		FShapedGlyphEntry& ShapedGlyphEntry = OutShapedLines.Last().GlyphsToRender[CurrentGlyphEntryIndex];
-		ShapedGlyphEntry.GlyphIndex = GlyphIndex;
-		ShapedGlyphEntry.SourceIndex = Index;
-		ShapedGlyphEntry.XAdvance = bIsZeroWidthSpace ? 0 : XAdvance;
-		ShapedGlyphEntry.YAdvance = 0;
-		ShapedGlyphEntry.XOffset = 0;
-		ShapedGlyphEntry.YOffset = 0;
-		ShapedGlyphEntry.Kerning = 0;
-		ShapedGlyphEntry.NumCharactersInGlyph = 1;
-		ShapedGlyphEntry.NumGraphemeClustersInGlyph = 1;
-		ShapedGlyphEntry.TextDirection = TextBiDi::ETextDirection::LeftToRight;
-		ShapedGlyphEntry.bIsVisible = !bIsWhitespace;
-
-		// Apply the kerning against the previous entry
-		if (CurrentGlyphEntryIndex > 0 && bHasKerning && ShapedGlyphEntry.bIsVisible)
-		{
-			FShapedGlyphEntry& PreviousShapedGlyphEntry = OutShapedLines.Last().GlyphsToRender[CurrentGlyphEntryIndex - 1];
-
-			FT_Vector KerningVector;
-			if (FT_Get_Kerning(Face, PreviousShapedGlyphEntry.GlyphIndex, ShapedGlyphEntry.GlyphIndex, FT_KERNING_DEFAULT, &KerningVector) == 0)
-			{
-				const int8 Kerning = KerningVector.x * FontInverseScale;
-				PreviousShapedGlyphEntry.XAdvance += Kerning;
-				PreviousShapedGlyphEntry.Kerning = Kerning;
-			}
+			FString BlockText;
+			Block->GetRun()->AppendTextTo(BlockText);
+
+			TSharedRef<const FString> BlockTextRef = MakeShared<const FString>(BlockText);
+			
+			const FTextRange BlockRange = Block->GetTextRange();
+			const FLayoutBlockTextContext BlockTextContext = Block->GetTextContext();
+
+			FShapedGlyphSequenceRef ShapedGlyphSequence = ShapedTextCacheUtil::GetShapedTextSubSequence(
+				BlockTextContext.ShapedTextCache,
+				FCachedShapedTextKey(Line.Range, 1.0f, BlockTextContext, Style.Font),
+				BlockRange,
+				**BlockTextRef,
+				BlockTextContext.TextDirection
+			);
+
+			ShapedLine.GlyphsToRender.Append(ShapedGlyphSequence->GetGlyphsToRender());
 		}
 	}
 }
-
-void FTextShaper::PerformHarfBuzzTextShaping(const FT_Face Face, const TCHAR* Text, int32 StartIndex, int32 EndIndex, TArray<FShapedGlyphLine>& OutShapedLines)
-{
-	hb_font_t* HarfBuzzFTFont = hb_ft_font_create(Face, nullptr);
-	hb_buffer_t* HarfBuzzBuffer = hb_buffer_create();
-
-	hb_unicode_funcs_t* HarfBuzzUnicodeFuncs = hb_unicode_funcs_get_default();
-
-
-	TArray<TTuple<int32, hb_script_t>> HarfBuzzSegments;
-	for (int32 Index = StartIndex; Index < EndIndex; Index++)
-	{
-		const hb_script_t CharHarfBuzzScript = hb_unicode_script(HarfBuzzUnicodeFuncs, Text[Index]);
-		if (HarfBuzzSegments.Num() > 0 && 
-			(CharHarfBuzzScript == HB_SCRIPT_COMMON || CharHarfBuzzScript == HB_SCRIPT_INHERITED || CharHarfBuzzScript == HB_SCRIPT_UNKNOWN))
-		{
-			continue;
-		}
-
-		if (HarfBuzzSegments.Num() == 0 || HarfBuzzSegments.Last().Value != CharHarfBuzzScript)
-		{
-			HarfBuzzSegments.Add(TTuple<int32, hb_script_t>(Index, CharHarfBuzzScript));
-		}
-	}
-
-	for (int32 Index = 0; Index < HarfBuzzSegments.Num(); Index++)
-	{
-		const TTuple<int32, hb_script_t>& HarfBuzzSegment = HarfBuzzSegments[Index];
-
-		int32 SegmentLength = EndIndex - HarfBuzzSegment.Key;
-		if (Index < HarfBuzzSegments.Num() - 1)
-	{
-			SegmentLength = HarfBuzzSegments[Index + 1].Key - HarfBuzzSegment.Key;
-	}
-
-		hb_buffer_add_utf16(HarfBuzzBuffer, (uint16_t*) Text, -1, HarfBuzzSegment.Key, SegmentLength);
-		hb_buffer_set_script(HarfBuzzBuffer, HarfBuzzSegment.Value);
-		hb_buffer_set_direction(HarfBuzzBuffer, HB_DIRECTION_RTL);
-	}
-
-	bool bHasKerning = FT_HAS_KERNING(Face) != 0;
-	const hb_feature_t HarfBuzzFeatures[] = { { HB_TAG('k','e','r','n'), bHasKerning, 0, uint32(-1) } };
-	const int32 HarfBuzzFeaturesCount = UE_ARRAY_COUNT(HarfBuzzFeatures);
-	hb_shape(HarfBuzzFTFont, HarfBuzzBuffer, HarfBuzzFeatures, HarfBuzzFeaturesCount);
-
-	uint32 glyph_count = 0;
-	hb_glyph_info_t * HarfBuzzGlyphInfo = hb_buffer_get_glyph_infos(HarfBuzzBuffer, &glyph_count);
-	hb_glyph_position_t * HarfBuzzGlyphPos = hb_buffer_get_glyph_positions(HarfBuzzBuffer, &glyph_count);
-
-	for (uint32 Index = 0; Index < glyph_count; Index++)
-	{
-		const int32 CurrentCharIndex = static_cast<int32>(HarfBuzzGlyphInfo[Index].cluster);
-		if (InsertSubstituteGlyphs(Face, Text, CurrentCharIndex, OutShapedLines))
-		{
-			continue;
-		}
-
-		const int32 CurrentGlyphEntryIndex = OutShapedLines.Last().GlyphsToRender.AddDefaulted();
-		FShapedGlyphEntry& ShapedGlyphEntry = OutShapedLines.Last().GlyphsToRender[CurrentGlyphEntryIndex];
-		ShapedGlyphEntry.GlyphIndex = HarfBuzzGlyphInfo[Index].codepoint;
-		ShapedGlyphEntry.SourceIndex = Index;
-		ShapedGlyphEntry.XAdvance = HarfBuzzGlyphPos[Index].x_advance * FontInverseScale;
-		ShapedGlyphEntry.YAdvance = HarfBuzzGlyphPos[Index].y_advance * FontInverseScale;
-		ShapedGlyphEntry.XOffset = HarfBuzzGlyphPos[Index].x_offset * FontInverseScale;
-		ShapedGlyphEntry.YOffset = HarfBuzzGlyphPos[Index].y_offset * FontInverseScale;
-		ShapedGlyphEntry.Kerning = 0;
-		ShapedGlyphEntry.NumCharactersInGlyph = 1;
-		ShapedGlyphEntry.NumGraphemeClustersInGlyph = 1;
-		ShapedGlyphEntry.TextDirection = TextBiDi::ETextDirection::LeftToRight;
-
-		const TCHAR CurrentChar = Text[CurrentCharIndex];
-		const bool bIsZeroWidthSpace = CurrentChar == TEXT('\u200B');
-		ShapedGlyphEntry.bIsVisible = !(bIsZeroWidthSpace || FText::IsWhitespace(CurrentChar));
-	}
-
-	hb_buffer_destroy(HarfBuzzBuffer);
-	hb_font_destroy(HarfBuzzFTFont);
-}
-
-bool FTextShaper::InsertSubstituteGlyphs(const FT_Face Face, const TCHAR* Text, const int32 Index, TArray<FShapedGlyphLine>& OutShapedLines)
-{
-	TCHAR Char = Text[Index];
-	if (TextBiDi::IsControlCharacter(Char))
-	{
-		// We insert a stub entry for control characters to avoid them being drawn as a visual glyph with size
-		const int32 CurrentGlyphEntryIndex = OutShapedLines.Last().GlyphsToRender.AddDefaulted();
-		FShapedGlyphEntry& ShapedGlyphEntry = OutShapedLines.Last().GlyphsToRender[CurrentGlyphEntryIndex];
-		//ShapedGlyphEntry.FontFaceData = InShapedGlyphFaceData;
-		ShapedGlyphEntry.GlyphIndex = 0;
-		ShapedGlyphEntry.SourceIndex = Index;
-		ShapedGlyphEntry.XAdvance = 0;
-		ShapedGlyphEntry.YAdvance = 0;
-		ShapedGlyphEntry.XOffset = 0;
-		ShapedGlyphEntry.YOffset = 0;
-		ShapedGlyphEntry.Kerning = 0;
-		ShapedGlyphEntry.NumCharactersInGlyph = 1;
-		ShapedGlyphEntry.NumGraphemeClustersInGlyph = 1;
-		ShapedGlyphEntry.TextDirection = TextBiDi::ETextDirection::LeftToRight;
-		ShapedGlyphEntry.bIsVisible = false;
-		return true;
-	}
-
-	if (Char == TEXT('\r'))
-	{
-		return true;
-	}
-
-	if (Char == TEXT('\n'))
-	{
-		OutShapedLines.AddDefaulted();
-		return true;
-	}
-
-	if (Char == TEXT('\t'))
-	{
-		int16 SpaceXAdvance = 0;
-		uint32 SpaceGlyphIndex = FT_Get_Char_Index(Face, TEXT(' '));
-
-		FT_Fixed AdvanceData = 0;
-		if (FT_Get_Advance(Face, SpaceGlyphIndex, 0, &AdvanceData) == 0)
-		{
-			SpaceXAdvance = ((AdvanceData + (1 << 9)) >> 10) * FontInverseScale;
-		}
-
-		// We insert a spacer glyph with (up-to) the width of 4 space glyphs in-place of a tab character
-		const int32 NumSpacesToInsert = 4 - (OutShapedLines.Last().GlyphsToRender.Num() % 4);
-		if (NumSpacesToInsert > 0)
-		{
-			const int32 CurrentGlyphEntryIndex = OutShapedLines.Last().GlyphsToRender.AddDefaulted();
-			FShapedGlyphEntry& ShapedGlyphEntry = OutShapedLines.Last().GlyphsToRender[CurrentGlyphEntryIndex];
-			ShapedGlyphEntry.GlyphIndex = SpaceGlyphIndex;
-			ShapedGlyphEntry.SourceIndex = Index;
-			ShapedGlyphEntry.XAdvance = SpaceXAdvance * NumSpacesToInsert;
-			ShapedGlyphEntry.YAdvance = 0;
-			ShapedGlyphEntry.XOffset = 0;
-			ShapedGlyphEntry.YOffset = 0;
-			ShapedGlyphEntry.Kerning = 0;
-			ShapedGlyphEntry.NumCharactersInGlyph = 1;
-			ShapedGlyphEntry.NumGraphemeClustersInGlyph = 1;
-			ShapedGlyphEntry.TextDirection = TextBiDi::ETextDirection::LeftToRight;
-			ShapedGlyphEntry.bIsVisible = false;
-		}
-
-		return true;
-	}
-
-	return false;
-}
diff --git a/Engine/Plugins/Experimental/Text3D/Source/Text3D/Private/TextShaper.h b/Engine/Plugins/Experimental/Text3D/Source/Text3D/Private/TextShaper.h
--- a/Engine/Plugins/Experimental/Text3D/Source/Text3D/Private/TextShaper.h
+++ b/Engine/Plugins/Experimental/Text3D/Source/Text3D/Private/TextShaper.h
@@ -1,20 +1,23 @@
 // Copyright Epic Games, Inc. All Rights Reserved.
 
 #pragma once
-#include "Text3DPrivate.h"
-#include "Fonts/FontCache.h"
 
 #include "Containers/Array.h"
-#include "Templates/UniquePtr.h"
+#include "Fonts/FontCache.h"
+#include "Framework/Text/TextLayout.h"
 #include "Internationalization/Text.h"
+#include "Templates/UniquePtr.h"
+#include "Text3DPrivate.h"
 
+class ITextLayoutMarshaller;
+class UFont;
 
 struct FShapedGlyphLine
 {
 	TArray<FShapedGlyphEntry> GlyphsToRender;
-	float Width;
+	float Width = 0.0f;
 
-	float GetAdvanced(const int32 Index, const float Kerning, const float WordSpacing) const
+	float GetAdvance(const int32 Index, const float Kerning, const float WordSpacing) const
 	{
 		check(Index >= 0 && Index < GlyphsToRender.Num());
 
@@ -39,37 +42,45 @@ struct FShapedGlyphLine
 		Width = 0.0f;
 		for (int32 Index = 0; Index < GlyphsToRender.Num(); Index++)
 		{
-			Width += GetAdvanced(Index, Kerning, WordSpacing);
+			Width += GetAdvance(Index, Kerning, WordSpacing);
 		}
 	}
+};
 
-	FShapedGlyphLine()
-	{
-		Width = 0.0f;
-	}
+class FText3DLayout : public FTextLayout
+{
+public:
+	FText3DLayout(const FTextBlockStyle& InStyle = FTextBlockStyle::GetDefault());
+
+protected:
+	FTextBlockStyle TextStyle;
+	
+	virtual TSharedRef<IRun> CreateDefaultTextRun(
+		const TSharedRef<FString>& NewText,
+		const FTextRange& NewRange) const override;
 };
 
 class FTextShaper final
 {
 public:
-	static FTextShaper* Get()							{ return Instance; }
+	static TSharedPtr<FTextShaper> Get();
 
-	void ShapeBidirectionalText(const FT_Face Face, const FString& Text, TArray<FShapedGlyphLine>& OutShapedLines);
+	void ShapeBidirectionalText(
+		const FTextBlockStyle& Style,
+		const FString& Text,
+		const TSharedPtr<FTextLayout>& TextLayout,
+		const TSharedPtr<ITextLayoutMarshaller>& TextMarshaller,
+		TArray<FShapedGlyphLine>& OutShapedLines);
 
-	static void Initialize();
-	static void Cleanup();
+private:
+	struct FPrivateToken { explicit FPrivateToken() = default; };
 
-private:	
-	FTextShaper();
-	FTextShaper(const FTextShaper&) = delete;
-	FTextShaper& operator=(const FTextShaper&) = delete;
-
-	void PerformKerningTextShaping(const FT_Face Face, const TCHAR* Text, const int32 StartIndex, const int32 EndIndex, TArray<FShapedGlyphLine>& OutShapedLines);
-	void PerformHarfBuzzTextShaping(const FT_Face Face, const TCHAR* Text, const int32 StartIndex, const int32 EndIndex, TArray<FShapedGlyphLine>& OutShapedLines);
-	bool InsertSubstituteGlyphs(const FT_Face Face, const TCHAR* Text, const int32 Index, TArray<FShapedGlyphLine>& OutShapedLines);
+public:
+	FTextShaper(FPrivateToken) { }
 
-	/** Unicode BiDi text detection */
-	TUniquePtr<TextBiDi::ITextBiDi> TextBiDiDetection;
+private:
+	friend class UText3DComponent;
 
-	static FTextShaper* Instance;
+	FTextShaper(const FTextShaper&) = delete;
+	FTextShaper& operator=(const FTextShaper&) = delete;
 };
diff --git a/Engine/Plugins/Experimental/Text3D/Source/Text3D/Public/Text3DComponent.h b/Engine/Plugins/Experimental/Text3D/Source/Text3D/Public/Text3DComponent.h
--- a/Engine/Plugins/Experimental/Text3D/Source/Text3D/Public/Text3DComponent.h
+++ b/Engine/Plugins/Experimental/Text3D/Source/Text3D/Public/Text3DComponent.h
@@ -10,6 +10,9 @@
 
 #include "Text3DComponent.generated.h"
 
+class FTextLayout;
+class ITextLayoutMarshaller;
+
 UENUM()
 enum class EText3DVerticalTextAlignment : uint8
 {
@@ -56,6 +59,10 @@ public:
 	virtual void PostTransacted(const FTransactionObjectEvent& TransactionEvent) override;
 #endif
 
+	/** Whether to allow automatic refresh/mesh generation */
+	UPROPERTY(EditAnywhere, BlueprintReadWrite, BlueprintSetter = SetRefreshOnChange, Category = "Text3D", AdvancedDisplay)
+	bool bRefreshOnChange = true;
+
 	/** The text to generate a 3d mesh */
 	UPROPERTY(EditAnywhere, BlueprintReadWrite, BlueprintSetter = SetText, Category = "Text3D", meta = (MultiLine = true))
 	FText Text;
@@ -161,6 +168,10 @@ public:
 	DECLARE_MULTICAST_DELEGATE(FTextGeneratedNative);
 	FTextGeneratedNative& OnTextGenerated() { return TextGeneratedNativeDelegate; }
 
+	/** Set whether to allow automatic refresh/mesh generation */
+	UFUNCTION(BlueprintCallable, Category = "Rendering|Components|Text3D")
+	void SetRefreshOnChange(const bool& Value);
+
 	/** Set the text value and signal the primitives to be rebuilt */
 	UFUNCTION(BlueprintCallable, Category = "Rendering|Components|Text3D")
 	void SetText(const FText& Value);
@@ -275,7 +286,7 @@ public:
 	UFUNCTION(BlueprintCallable, Category = "Rendering|Components|Text3D")
 	UStaticMeshComponent* GetGlyphMeshComponent(int32 Index);
 
-	/** Gets all the glyph meshes*/
+	/** Gets all the glyph meshes */
 	UFUNCTION(BlueprintCallable, Category = "Rendering|Components|Text3D")
 	const TArray<UStaticMeshComponent*>& GetGlyphMeshComponents();
 	
@@ -283,10 +294,15 @@ public:
 	UFUNCTION(BlueprintCallable, Category = "Rendering|Components|Text3D")
 	FVector GetTextScale();
 
+	/** Manually update the geometry, ignoring RefreshOnChange (but still accounting for the Freeze flag) */
+	void Rebuild();
+
 protected:
 	virtual void OnVisibilityChanged() override;
 	virtual void OnHiddenInGameChanged() override;
 
+	void ClearTextMesh();
+
 private:
 	UPROPERTY()
 	TObjectPtr<class USceneComponent> TextRoot;
@@ -307,6 +323,9 @@ private:
 	TSharedPtr<struct FText3DShapedText> ShapedText;
 	TArray<TSharedPtr<int32>> CachedCounterReferences;
 
+	TSharedPtr<FTextLayout> TextLayout;
+	TSharedPtr<ITextLayoutMarshaller> TextLayoutMarshaller;
+
 	UPROPERTY(Transient)
 	TArray<TObjectPtr<USceneComponent>> CharacterKernings;
 
@@ -320,11 +339,9 @@ private:
 	void SetMaterial(const EText3DGroupType Type, class UMaterialInterface* Material);
 
 	void UpdateTransforms();
-	void Rebuild(const bool bCleanCache = false);
-	void Update();
-	void ClearTextMesh();
-	void BuildTextMesh(const bool bCleanCache = false);
-	void BuildTextMeshInternal(const bool bCleanCache);
+	void RebuildInternal(const bool& bIsAutoUpdate = true, const bool& bCleanCache = false);
+	void BuildTextMesh(const bool& bCleanCache = false);
+	void BuildTextMeshInternal(const bool& bCleanCache);
 	void CheckBevel();
 	float MaxBevel() const;
 
diff --git a/Engine/Plugins/Experimental/Text3D/Source/Text3D/Text3D.Build.cs b/Engine/Plugins/Experimental/Text3D/Source/Text3D/Text3D.Build.cs
--- a/Engine/Plugins/Experimental/Text3D/Source/Text3D/Text3D.Build.cs
+++ b/Engine/Plugins/Experimental/Text3D/Source/Text3D/Text3D.Build.cs
@@ -1,5 +1,6 @@
 // Copyright Epic Games, Inc. All Rights Reserved.
 
+using System.IO;
 using UnrealBuildTool;
 
 public class Text3D : ModuleRules
@@ -17,13 +18,21 @@ public class Text3D : ModuleRules
 
 		PrivateDependencyModuleNames.AddRange(new string[] {
 			"FreeType2",
-			"GeometryCore",
 			"GeometryAlgorithms",
-            "HarfBuzz",
+			"GeometryCore",
+			"HarfBuzz",
             "ICU",
             "MeshDescription",
+            "Slate",
 			"SlateCore",
 			"StaticMeshDescription",
 		});
+
+		// Needed to reference underlying FreeType info
+		// @todo: Move font vector handling to new module ("FontCore") - don't expose FT directly?
+		string EnginePath = Path.GetFullPath(Target.RelativeEnginePath);
+		PrivateIncludePaths.AddRange(new string[]  {
+			Path.Combine(EnginePath, "Source/Runtime/SlateCore/Private")
+		});
 	}
 }
