Text3D -game crash fix

diff --git a/Engine/Plugins/Experimental/Text3D/Source/Text3D/Private/Text3DComponent.cpp b/Engine/Plugins/Experimental/Text3D/Source/Text3D/Private/Text3DComponent.cpp
--- a/Engine/Plugins/Experimental/Text3D/Source/Text3D/Private/Text3DComponent.cpp
+++ b/Engine/Plugins/Experimental/Text3D/Source/Text3D/Private/Text3DComponent.cpp
@@ -51,9 +51,7 @@ UText3DComponent::UText3DComponent()
 	, ShapedText(new FText3DShapedText())
 {
 	TextRoot = CreateDefaultSubobject<USceneComponent>(TEXT("TextRoot"));
-#if WITH_EDITOR
-	TextRoot->SetIsVisualizationComponent(true);
-#endif
+	TextRoot->SetupAttachment(this);
 
 	if (!IsRunningDedicatedServer())
 	{
