[Text3D] Fixed material assignments not updating

diff --git a/Engine/Plugins/Experimental/Text3D/Source/Text3D/Private/Text3DComponent.cpp b/Engine/Plugins/Experimental/Text3D/Source/Text3D/Private/Text3DComponent.cpp
--- a/Engine/Plugins/Experimental/Text3D/Source/Text3D/Private/Text3DComponent.cpp
+++ b/Engine/Plugins/Experimental/Text3D/Source/Text3D/Private/Text3DComponent.cpp
@@ -159,11 +159,28 @@ void UText3DComponent::OnUnregister()
 void UText3DComponent::PostEditChangeProperty(FPropertyChangedEvent& PropertyChangedEvent)
 {
 	Super::PostEditChangeProperty(PropertyChangedEvent);
-	
-	static FName BevelTypePropertyName = GET_MEMBER_NAME_CHECKED(UText3DComponent, BevelType);
-	static FName BevelSegmentsPropertyName = GET_MEMBER_NAME_CHECKED(UText3DComponent, BevelSegments);
+
+	static TMap<FName, EText3DGroupType> MaterialToGroup =
+	{
+		{ GET_MEMBER_NAME_CHECKED(UText3DComponent, FrontMaterial), EText3DGroupType::Front },
+		{ GET_MEMBER_NAME_CHECKED(UText3DComponent, BackMaterial), EText3DGroupType::Back },
+		{ GET_MEMBER_NAME_CHECKED(UText3DComponent, ExtrudeMaterial), EText3DGroupType::Extrude },
+		{ GET_MEMBER_NAME_CHECKED(UText3DComponent, BevelMaterial), EText3DGroupType::Bevel }
+	};
+
+	using FGetter = TDelegate<const TObjectPtr<UMaterialInterface>&()>;
+	static TMap<FName, FGetter> MaterialLookup =
+	{
+		{ GET_MEMBER_NAME_CHECKED(UText3DComponent, FrontMaterial), FGetter::CreateUObject(this, &UText3DComponent::GetFrontMaterial) },
+		{ GET_MEMBER_NAME_CHECKED(UText3DComponent, BackMaterial), FGetter::CreateUObject(this, &UText3DComponent::GetBackMaterial) },
+		{ GET_MEMBER_NAME_CHECKED(UText3DComponent, ExtrudeMaterial), FGetter::CreateUObject(this, &UText3DComponent::GetExtrudeMaterial) },
+		{ GET_MEMBER_NAME_CHECKED(UText3DComponent, BevelMaterial), FGetter::CreateUObject(this, &UText3DComponent::GetBevelMaterial) }
+	};
 
 	const FName Name = PropertyChangedEvent.GetPropertyName();
+
+	static FName BevelTypePropertyName = GET_MEMBER_NAME_CHECKED(UText3DComponent, BevelType);
+	static FName BevelSegmentsPropertyName = GET_MEMBER_NAME_CHECKED(UText3DComponent, BevelSegments);
 	if (Name == BevelTypePropertyName)
 	{
 		switch (BevelType)
@@ -218,6 +235,16 @@ void UText3DComponent::PostEditChangeProperty(FPropertyChangedEvent& PropertyCha
 	{
 		MarkForLayoutUpdate();
 	}
+	else if (const EText3DGroupType* MaterialGroup = MaterialToGroup.Find(Name))
+	{
+		UpdateMaterial(*MaterialGroup, MaterialLookup[Name].Execute());
+	}
+	// Any property not explicitly handled should trigger a full rebuild
+	else
+	{
+		MarkForGeometryUpdate();
+		MarkForLayoutUpdate();
+	}
 
 	RebuildInternal();
 }
diff --git a/Engine/Plugins/Experimental/Text3D/Source/Text3D/Public/Text3DComponent.h b/Engine/Plugins/Experimental/Text3D/Source/Text3D/Public/Text3DComponent.h
--- a/Engine/Plugins/Experimental/Text3D/Source/Text3D/Public/Text3DComponent.h
+++ b/Engine/Plugins/Experimental/Text3D/Source/Text3D/Public/Text3DComponent.h
@@ -92,19 +92,19 @@ public:
 	float OutlineExpand;
 
 	/** Material for the front part */
-	UPROPERTY(EditAnywhere, BlueprintReadWrite, BlueprintSetter = SetFrontMaterial, Category = "Text3D")
+	UPROPERTY(EditAnywhere, BlueprintReadWrite, Getter = "Auto", BlueprintSetter = SetFrontMaterial, Category = "Text3D")
 	TObjectPtr<class UMaterialInterface> FrontMaterial;
 
 	/** Material for the bevel part */
-	UPROPERTY(EditAnywhere, BlueprintReadWrite, BlueprintSetter = SetBevelMaterial, Category = "Text3D", Meta = (EditCondition = "!bOutline"))
+	UPROPERTY(EditAnywhere, BlueprintReadWrite, Getter = "Auto", BlueprintSetter = SetBevelMaterial, Category = "Text3D", Meta = (EditCondition = "!bOutline"))
 	TObjectPtr<class UMaterialInterface> BevelMaterial;
 
 	/** Material for the extruded part */
-	UPROPERTY(EditAnywhere, BlueprintReadWrite, BlueprintSetter = SetExtrudeMaterial, Category = "Text3D")
+	UPROPERTY(EditAnywhere, BlueprintReadWrite, Getter = "Auto", BlueprintSetter = SetExtrudeMaterial, Category = "Text3D")
 	TObjectPtr<class UMaterialInterface> ExtrudeMaterial;
 
 	/** Material for the back part */
-	UPROPERTY(EditAnywhere, BlueprintReadWrite, BlueprintSetter = SetBackMaterial, Category = "Text3D")
+	UPROPERTY(EditAnywhere, BlueprintReadWrite, Getter = "Auto", BlueprintSetter = SetBackMaterial, Category = "Text3D")
 	TObjectPtr<class UMaterialInterface> BackMaterial;
 
 	/** Text font */
