[Avalanche] Text refresh on load

diff --git a/Engine/Plugins/Experimental/Text3D/Source/Text3D/Private/Text3DComponent.cpp b/Engine/Plugins/Experimental/Text3D/Source/Text3D/Private/Text3DComponent.cpp
--- a/Engine/Plugins/Experimental/Text3D/Source/Text3D/Private/Text3DComponent.cpp
+++ b/Engine/Plugins/Experimental/Text3D/Source/Text3D/Private/Text3DComponent.cpp
@@ -101,6 +101,14 @@ UText3DComponent::UText3DComponent()
 	TextScale = FVector::ZeroVector;
 }
 
+void UText3DComponent::PostLoad()
+{
+	// Reset so it's rebuilt (needed if re-using the component!)
+	ModifyFlags = EText3DModifyFlags::All;
+	
+	Super::PostLoad();
+}
+
 void UText3DComponent::BeginDestroy()
 {
 	ClearTextMesh();
@@ -618,7 +626,7 @@ void UText3DComponent::SetFreeze(const bool bFreeze)
 	bFreezeBuild = bFreeze;
 	if (bFreeze)
 	{
-		ModifyFlags = EText3DModifyFlags::Unfreeze;
+		ModifyFlags |= EText3DModifyFlags::Unfreeze;
 	}
 	else if (EnumHasAnyFlags(ModifyFlags,EText3DModifyFlags::Unfreeze))
 	{
diff --git a/Engine/Plugins/Experimental/Text3D/Source/Text3D/Public/Text3DComponent.h b/Engine/Plugins/Experimental/Text3D/Source/Text3D/Public/Text3DComponent.h
--- a/Engine/Plugins/Experimental/Text3D/Source/Text3D/Public/Text3DComponent.h
+++ b/Engine/Plugins/Experimental/Text3D/Source/Text3D/Public/Text3DComponent.h
@@ -50,6 +50,7 @@ class TEXT3D_API UText3DComponent : public USceneComponent
 public:
 	UText3DComponent();
 
+	virtual void PostLoad() override;
 	virtual void BeginDestroy() override;
 	virtual void OnRegister() override;
 	virtual void OnUnregister() override;
