Text3D cleanup of recent submission 25696117

diff --git a/Engine/Plugins/Experimental/Text3D/Source/Text3D/Private/Text3DComponent.cpp b/Engine/Plugins/Experimental/Text3D/Source/Text3D/Private/Text3DComponent.cpp
--- a/Engine/Plugins/Experimental/Text3D/Source/Text3D/Private/Text3DComponent.cpp
+++ b/Engine/Plugins/Experimental/Text3D/Source/Text3D/Private/Text3DComponent.cpp
@@ -3,6 +3,7 @@
 #include "Text3DComponent.h"
 
 #include "Algo/Count.h"
+#include "Algo/IndexOf.h"
 #include "Async/Async.h"
 #include "Components/StaticMeshComponent.h"
 #include "Engine/Engine.h"
@@ -176,7 +177,7 @@ void UText3DComponent::PostEditChangeProperty(FPropertyChangedEvent& PropertyCha
 		{ GET_MEMBER_NAME_CHECKED(UText3DComponent, BevelMaterial), EText3DGroupType::Bevel }
 	};
 
-	using FGetter = TDelegate<const TObjectPtr<UMaterialInterface>&()>;
+	using FGetter = TDelegate<UMaterialInterface*()>;
 	static TMap<FName, FGetter> MaterialLookup =
 	{
 		{ GET_MEMBER_NAME_CHECKED(UText3DComponent, FrontMaterial), FGetter::CreateUObject(this, &UText3DComponent::GetFrontMaterial) },
@@ -269,7 +270,12 @@ void UText3DComponent::PostTransacted(const FTransactionObjectEvent& Transaction
 }
 #endif
 
-void UText3DComponent::SetRefreshOnChange(const bool& Value)
+bool UText3DComponent::RefreshesOnChange() const
+{
+	return bRefreshOnChange;
+}
+
+void UText3DComponent::SetRefreshOnChange(const bool Value)
 {
 	if (bRefreshOnChange != Value)
 	{
@@ -277,6 +283,11 @@ void UText3DComponent::SetRefreshOnChange(const bool& Value)
 	}
 }
 
+const FText& UText3DComponent::GetText() const
+{
+	return Text;
+}
+
 void UText3DComponent::SetText(const FText& Value)
 {
 	if (!Text.EqualTo(Value))
@@ -287,7 +298,12 @@ void UText3DComponent::SetText(const FText& Value)
 	}
 }
 
-void UText3DComponent::SetFont(UFont* const InFont)
+const UFont* UText3DComponent::GetFont() const
+{
+	return Font;
+}
+
+void UText3DComponent::SetFont(UFont* InFont)
 {
 	if (Font != InFont)
 	{
@@ -297,7 +313,12 @@ void UText3DComponent::SetFont(UFont* const InFont)
 	}
 }
 
-void UText3DComponent::SetOutline(const bool bValue)
+bool UText3DComponent::HasOutline() const
+{
+	return bOutline;
+}
+
+void UText3DComponent::SetHasOutline(const bool bValue)
 {
 	if (bOutline != bValue)
 	{
@@ -307,6 +328,11 @@ void UText3DComponent::SetOutline(const bool bValue)
 	}
 }
 
+float UText3DComponent::GetOutlineExpand() const
+{
+	return OutlineExpand;
+}
+
 void UText3DComponent::SetOutlineExpand(const float Value)
 {
 	const float NewValue = Value;
@@ -318,6 +344,11 @@ void UText3DComponent::SetOutlineExpand(const float Value)
 	}
 }
 
+float UText3DComponent::GetExtrude() const
+{
+	return Extrude;
+}
+
 void UText3DComponent::SetExtrude(const float Value)
 {
 	const float NewValue = FMath::Max(0.0f, Value);
@@ -330,10 +361,14 @@ void UText3DComponent::SetExtrude(const float Value)
 	}
 }
 
+float UText3DComponent::GetBevel() const
+{
+	return Bevel;
+}
+
 void UText3DComponent::SetBevel(const float Value)
 {
 	const float NewValue = FMath::Clamp(Value, 0.f, MaxBevel());
-
 	if (!FMath::IsNearlyEqual(Bevel, NewValue))
 	{
 		Bevel = NewValue;
@@ -342,6 +377,11 @@ void UText3DComponent::SetBevel(const float Value)
 	}
 }
 
+EText3DBevelType UText3DComponent::GetBevelType() const
+{
+	return BevelType;
+}
+
 void UText3DComponent::SetBevelType(const EText3DBevelType Value)
 {
 	if (BevelType != Value)
@@ -352,6 +392,11 @@ void UText3DComponent::SetBevelType(const EText3DBevelType Value)
 	}
 }
 
+int32 UText3DComponent::GetBevelSegments() const
+{
+	return BevelSegments;
+}
+
 void UText3DComponent::SetBevelSegments(const int32 Value)
 {
 	int32 MinBevelSegments = 1;
@@ -369,21 +414,41 @@ void UText3DComponent::SetBevelSegments(const int32 Value)
 	}
 }
 
+UMaterialInterface* UText3DComponent::GetFrontMaterial() const
+{
+	return FrontMaterial;
+}
+
 void UText3DComponent::SetFrontMaterial(UMaterialInterface* Value)
 {
 	SetMaterial(EText3DGroupType::Front, Value);
 }
 
+UMaterialInterface* UText3DComponent::GetBevelMaterial() const
+{
+	return BevelMaterial;
+}
+
 void UText3DComponent::SetBevelMaterial(UMaterialInterface* Value)
 {
 	SetMaterial(EText3DGroupType::Bevel, Value);
 }
 
+UMaterialInterface* UText3DComponent::GetExtrudeMaterial() const
+{
+	return ExtrudeMaterial;
+}
+
 void UText3DComponent::SetExtrudeMaterial(UMaterialInterface* Value)
 {
 	SetMaterial(EText3DGroupType::Extrude, Value);
 }
 
+UMaterialInterface* UText3DComponent::GetBackMaterial() const
+{
+	return BackMaterial;
+}
+
 void UText3DComponent::SetBackMaterial(UMaterialInterface* Value)
 {
 	SetMaterial(EText3DGroupType::Back, Value);
@@ -529,6 +594,11 @@ void UText3DComponent::SetMaterial(const EText3DGroupType Type, UMaterialInterfa
 	}
 }
 
+float UText3DComponent::GetKerning() const
+{
+	return Kerning;
+}
+
 void UText3DComponent::SetKerning(const float Value)
 {
 	if (!FMath::IsNearlyEqual(Kerning, Value))
@@ -538,6 +608,11 @@ void UText3DComponent::SetKerning(const float Value)
 	}
 }
 
+float UText3DComponent::GetLineSpacing() const
+{
+	return LineSpacing;
+}
+
 void UText3DComponent::SetLineSpacing(const float Value)
 {
 	if (!FMath::IsNearlyEqual(LineSpacing, Value))
@@ -547,6 +622,11 @@ void UText3DComponent::SetLineSpacing(const float Value)
 	}
 }
 
+float UText3DComponent::GetWordSpacing() const
+{
+	return WordSpacing;
+}
+
 void UText3DComponent::SetWordSpacing(const float Value)
 {
 	if (!FMath::IsNearlyEqual(WordSpacing, Value))
@@ -556,6 +636,11 @@ void UText3DComponent::SetWordSpacing(const float Value)
 	}
 }
 
+EText3DHorizontalTextAlignment UText3DComponent::GetHorizontalAlignment() const
+{
+	return HorizontalAlignment;
+}
+
 void UText3DComponent::SetHorizontalAlignment(const EText3DHorizontalTextAlignment Value)
 {
 	if (HorizontalAlignment != Value)
@@ -565,6 +650,11 @@ void UText3DComponent::SetHorizontalAlignment(const EText3DHorizontalTextAlignme
 	}
 }
 
+EText3DVerticalTextAlignment UText3DComponent::GetVerticalAlignment() const
+{
+	return VerticalAlignment;
+}
+
 void UText3DComponent::SetVerticalAlignment(const EText3DVerticalTextAlignment Value)
 {
 	if (VerticalAlignment != Value)
@@ -574,6 +664,11 @@ void UText3DComponent::SetVerticalAlignment(const EText3DVerticalTextAlignment V
 	}
 }
 
+bool UText3DComponent::HasMaxWidth() const
+{
+	return bHasMaxWidth;
+}
+
 void UText3DComponent::SetHasMaxWidth(const bool Value)
 {
 	if (bHasMaxWidth != Value)
@@ -583,6 +678,11 @@ void UText3DComponent::SetHasMaxWidth(const bool Value)
 	}
 }
 
+float UText3DComponent::GetMaxWidth() const
+{
+	return MaxWidth;
+}
+
 void UText3DComponent::SetMaxWidth(const float Value)
 {
 	const float NewValue = FMath::Max(1.0f, Value);
@@ -593,6 +693,11 @@ void UText3DComponent::SetMaxWidth(const float Value)
 	}
 }
 
+bool UText3DComponent::HasMaxHeight() const
+{
+	return bHasMaxHeight;
+}
+
 void UText3DComponent::SetHasMaxHeight(const bool Value)
 {
 	if (bHasMaxHeight != Value)
@@ -602,6 +707,11 @@ void UText3DComponent::SetHasMaxHeight(const bool Value)
 	}
 }
 
+float UText3DComponent::GetMaxHeight() const
+{
+	return MaxHeight;
+}
+
 void UText3DComponent::SetMaxHeight(const float Value)
 {
 	const float NewValue = FMath::Max(1.0f, Value);
@@ -612,6 +722,11 @@ void UText3DComponent::SetMaxHeight(const float Value)
 	}
 }
 
+bool UText3DComponent::ScalesProportionally() const
+{
+	return bScaleProportionally;
+}
+
 void UText3DComponent::SetScaleProportionally(const bool Value)
 {
 	if (bScaleProportionally != Value)
@@ -621,6 +736,11 @@ void UText3DComponent::SetScaleProportionally(const bool Value)
 	}
 }
 
+bool UText3DComponent::IsFrozen() const
+{
+	return bFreezeBuild;
+}
+
 void UText3DComponent::SetFreeze(const bool bFreeze)
 {
 	bFreezeBuild = bFreeze;
@@ -634,6 +754,11 @@ void UText3DComponent::SetFreeze(const bool bFreeze)
 	}
 }
 
+bool UText3DComponent::CastsShadow() const
+{
+	return bCastShadow;
+}
+
 void UText3DComponent::SetCastShadow(bool NewCastShadow)
 {
 	if (NewCastShadow != bCastShadow)
@@ -972,7 +1097,7 @@ void UText3DComponent::BuildTextMeshInternal(const bool& bCleanCache)
 		},
 		0);
 
-	TMap<uint32, TSharedPtr<FFreeTypeFace>> GlyphIndexToFontFace;
+	TMap<uint32, const FFreeTypeFace*> GlyphIndexToFontFace;
 	GlyphIndexToFontFace.Reserve(ApproximateGlyphNum);
 
 	for (int32 LineIndex = 0; LineIndex < ShapedText->Lines.Num(); ++LineIndex)
@@ -989,8 +1114,12 @@ void UText3DComponent::BuildTextMeshInternal(const bool& bCleanCache)
 				
 				continue;
 			}
-			
-			GlyphIndexToFontFace.FindOrAdd(GlyphEntry.GlyphIndex, GlyphEntry.FontFaceData->FontFace.Pin());
+
+			if (const TSharedPtr<FFreeTypeFace> FontFacePtr = GlyphEntry.FontFaceData->FontFace.Pin();
+				FontFacePtr.IsValid())
+			{
+				GlyphIndexToFontFace.FindOrAdd(GlyphEntry.GlyphIndex, FontFacePtr.Get());	
+			}
 		}
 	}
 	
@@ -1028,7 +1157,7 @@ void UText3DComponent::BuildTextMeshInternal(const bool& bCleanCache)
 			// Count even when mesh is nullptr (allocation still creates components to avoid mesh building in allocation step)
 			const int32 GlyphId = GlyphIndex++;
 
-			TSharedPtr<FFreeTypeFace> FontFace = GlyphIndexToFontFace[ShapedGlyph.GlyphIndex];
+			const FFreeTypeFace* FontFace = GlyphIndexToFontFace[ShapedGlyph.GlyphIndex];
 			UStaticMesh* CachedMesh = CachedFontData.GetGlyphMesh(ShapedGlyph.GlyphIndex, FGlyphMeshParameters{Extrude, Bevel, BevelType, BevelSegments, bOutline, OutlineExpand}, FontFace);
 			if (!CachedMesh || FMath::IsNearlyZero(CachedMesh->GetBounds().SphereRadius))
 			{
diff --git a/Engine/Plugins/Experimental/Text3D/Source/Text3D/Private/Text3DEngineSubsystem.cpp b/Engine/Plugins/Experimental/Text3D/Source/Text3D/Private/Text3DEngineSubsystem.cpp
--- a/Engine/Plugins/Experimental/Text3D/Source/Text3D/Private/Text3DEngineSubsystem.cpp
+++ b/Engine/Plugins/Experimental/Text3D/Source/Text3D/Private/Text3DEngineSubsystem.cpp
@@ -195,7 +195,7 @@ TSharedPtr<int32> FCachedFontData::GetMeshesCacheCounter(const FGlyphMeshParamet
 	return CachedMeshes.GetCacheCounter();
 }
 
-UStaticMesh* FCachedFontData::GetGlyphMesh(uint32 GlyphIndex, const FGlyphMeshParameters& Parameters, const TSharedPtr<FFreeTypeFace>& FontFaceData)
+UStaticMesh* FCachedFontData::GetGlyphMesh(uint32 GlyphIndex, const FGlyphMeshParameters& Parameters, const FFreeTypeFace* FontFaceData)
 {
 	const uint32 HashParameters = GetTypeHash(Parameters);
 	FCachedFontMeshes& CachedMeshes = Meshes.FindOrAdd(HashParameters);
@@ -241,9 +241,9 @@ const FString& FCachedFontData::GetFontName()
 	return FontName;
 }
 
-TSharedContourNode FCachedFontData::GetGlyphContours(uint32 GlyphIndex, const TSharedPtr<FFreeTypeFace>& FontFaceData)
+TSharedContourNode FCachedFontData::GetGlyphContours(uint32 GlyphIndex, const FFreeTypeFace* FontFaceData)
 {
-	const FT_Face& GlyphFreeTypeFace = FontFaceData.IsValid() && FontFaceData->IsFaceValid()
+	const FT_Face& GlyphFreeTypeFace = FontFaceData && FontFaceData->IsFaceValid()
 		? FontFaceData->GetFace()
 		: FreeTypeFace;
 	check(GlyphFreeTypeFace);
diff --git a/Engine/Plugins/Experimental/Text3D/Source/Text3D/Private/Text3DEngineSubsystem.h b/Engine/Plugins/Experimental/Text3D/Source/Text3D/Private/Text3DEngineSubsystem.h
--- a/Engine/Plugins/Experimental/Text3D/Source/Text3D/Private/Text3DEngineSubsystem.h
+++ b/Engine/Plugins/Experimental/Text3D/Source/Text3D/Private/Text3DEngineSubsystem.h
@@ -89,8 +89,8 @@ public:
 	TSharedPtr<int32> GetCacheCounter();
 	TSharedPtr<int32> GetMeshesCacheCounter(const FGlyphMeshParameters& Parameters);
 
-	UStaticMesh* GetGlyphMesh(uint32 GlyphIndex, const FGlyphMeshParameters& Parameters, const TSharedPtr<FFreeTypeFace>& FontFaceData = nullptr);
-	TSharedContourNode GetGlyphContours(uint32 GlyphIndex, const TSharedPtr<FFreeTypeFace>& FontFaceData = nullptr);
+	UStaticMesh* GetGlyphMesh(uint32 GlyphIndex, const FGlyphMeshParameters& Parameters, const FFreeTypeFace* FontFaceData = nullptr);
+	TSharedContourNode GetGlyphContours(uint32 GlyphIndex, const FFreeTypeFace* FontFaceData = nullptr);
 
 	UPROPERTY()
 	TObjectPtr<class UFont> Font;
diff --git a/Engine/Plugins/Experimental/Text3D/Source/Text3D/Private/TextShaper.h b/Engine/Plugins/Experimental/Text3D/Source/Text3D/Private/TextShaper.h
--- a/Engine/Plugins/Experimental/Text3D/Source/Text3D/Private/TextShaper.h
+++ b/Engine/Plugins/Experimental/Text3D/Source/Text3D/Private/TextShaper.h
@@ -12,11 +12,16 @@
 class ITextLayoutMarshaller;
 class UFont;
 
+/** Contains a single line of text with sufficient information to fetch and transform each character. */
 struct FShapedGlyphLine
 {
+	/** The corresponding shaped glyph for each character in this line of text. */
 	TArray<FShapedGlyphEntry> GlyphsToRender;
+
+	/** Stored result of CalculateWidth. */
 	float Width = 0.0f;
 
+	/** Get's the offset from the previous character, accounting for kerning and word spacing. */
 	float GetAdvance(const int32 Index, const float Kerning, const float WordSpacing) const
 	{
 		check(Index >= 0 && Index < GlyphsToRender.Num());
@@ -37,6 +42,7 @@ struct FShapedGlyphLine
 		return Advance;
 	}
 
+	/** Calculates the total width of this line. */
 	void CalculateWidth(const float Kerning, const float WordSpacing)
 	{
 		Width = 0.0f;
@@ -47,24 +53,34 @@ struct FShapedGlyphLine
 	}
 };
 
+/** An implementation of FTextLayout, which discards most Slate widget specific functionality. */
 class FText3DLayout : public FTextLayout
 {
 public:
+	/** Optionally provide a custom TextBlock style. */
 	FText3DLayout(const FTextBlockStyle& InStyle = FTextBlockStyle::GetDefault());
 
 protected:
+	/** Parameters relevant to text layout. */
 	FTextBlockStyle TextStyle;
-	
+
+	/** Required but unused override. */
 	virtual TSharedRef<IRun> CreateDefaultTextRun(
 		const TSharedRef<FString>& NewText,
 		const FTextRange& NewRange) const override;
 };
 
+/** A singleton that handles shaping operations and writes the result to a provided TextLayout. */
 class FTextShaper final
 {
 public:
+	/** Returns TextShaper singleton. */
 	static TSharedPtr<FTextShaper> Get();
 
+	/**
+	 * Arranges the provided text to match the requested layout, accounting for scale, offsets etc.
+	 * Analogous to FSlateFontCache::ShapeBidirectionalText.
+	 */
 	void ShapeBidirectionalText(
 		const FTextBlockStyle& Style,
 		const FString& Text,
diff --git a/Engine/Plugins/Experimental/Text3D/Source/Text3D/Public/Text3DComponent.h b/Engine/Plugins/Experimental/Text3D/Source/Text3D/Public/Text3DComponent.h
--- a/Engine/Plugins/Experimental/Text3D/Source/Text3D/Public/Text3DComponent.h
+++ b/Engine/Plugins/Experimental/Text3D/Source/Text3D/Public/Text3DComponent.h
@@ -12,6 +12,8 @@
 
 class FTextLayout;
 class ITextLayoutMarshaller;
+class UFont;
+class UMaterialInterface;
 
 UENUM()
 enum class EText3DVerticalTextAlignment : uint8
@@ -42,7 +44,7 @@ enum class EText3DModifyFlags : uint8
 };
 ENUM_CLASS_FLAGS(EText3DModifyFlags)
 
-UCLASS(ClassGroup = (Text3D), meta=(BlueprintSpawnableComponent))
+UCLASS(ClassGroup = (Text3D), meta = (BlueprintSpawnableComponent))
 class TEXT3D_API UText3DComponent : public USceneComponent
 {
 	GENERATED_BODY()
@@ -50,8 +52,8 @@ class TEXT3D_API UText3DComponent : public USceneComponent
 public:
 	UText3DComponent();
 
-	virtual void PostLoad() override;
 	virtual void BeginDestroy() override;
+	virtual void PostLoad() override;
 	virtual void OnRegister() override;
 	virtual void OnUnregister() override;
 
@@ -60,104 +62,6 @@ public:
 	virtual void PostTransacted(const FTransactionObjectEvent& TransactionEvent) override;
 #endif
 
-	/** Whether to allow automatic refresh/mesh generation */
-	UPROPERTY(EditAnywhere, BlueprintReadWrite, BlueprintSetter = SetRefreshOnChange, Category = "Text3D", AdvancedDisplay)
-	bool bRefreshOnChange = true;
-
-	/** The text to generate a 3d mesh */
-	UPROPERTY(EditAnywhere, BlueprintReadWrite, BlueprintSetter = SetText, Category = "Text3D", meta = (MultiLine = true))
-	FText Text;
-
-	/** Size of the extrude */
-	UPROPERTY(EditAnywhere, BlueprintReadWrite, BlueprintSetter = SetExtrude, Category = "Text3D", Meta = (ClampMin = 0))
-	float Extrude;
-
-	/** Size of bevel */
-	UPROPERTY(EditAnywhere, BlueprintReadWrite, BlueprintSetter = SetBevel, Category = "Text3D", Meta = (EditCondition = "!bOutline", ClampMin = 0))
-	float Bevel;
-
-	/** Bevel Type */
-	UPROPERTY(EditAnywhere, BlueprintReadWrite, BlueprintSetter = SetBevelType, Category = "Text3D", Meta = (EditCondition = "!bOutline"))
-	EText3DBevelType BevelType;
-
-	/** Bevel Segments (Defines the amount of tesselation for the bevel part) */
-	UPROPERTY(EditAnywhere, BlueprintReadWrite, BlueprintSetter = SetBevelSegments, Category = "Text3D", Meta = (EditCondition = "!bOutline", ClampMin = 1, ClampMax = 15))
-	int32 BevelSegments;
-
-	/** Generate Outline */
-	UPROPERTY(EditAnywhere, BlueprintReadWrite, BlueprintSetter = SetOutline, Category = "Text3D")
-	bool bOutline;
-
-	/** Outline expand/offset amount */
-	UPROPERTY(EditAnywhere, BlueprintReadWrite, BlueprintSetter = SetOutlineExpand, Category = "Text3D", Meta = (EditCondition = "bOutline"))
-	float OutlineExpand;
-
-	/** Material for the front part */
-	UPROPERTY(EditAnywhere, BlueprintReadWrite, Getter = "Auto", BlueprintSetter = SetFrontMaterial, Category = "Text3D")
-	TObjectPtr<class UMaterialInterface> FrontMaterial;
-
-	/** Material for the bevel part */
-	UPROPERTY(EditAnywhere, BlueprintReadWrite, Getter = "Auto", BlueprintSetter = SetBevelMaterial, Category = "Text3D", Meta = (EditCondition = "!bOutline"))
-	TObjectPtr<class UMaterialInterface> BevelMaterial;
-
-	/** Material for the extruded part */
-	UPROPERTY(EditAnywhere, BlueprintReadWrite, Getter = "Auto", BlueprintSetter = SetExtrudeMaterial, Category = "Text3D")
-	TObjectPtr<class UMaterialInterface> ExtrudeMaterial;
-
-	/** Material for the back part */
-	UPROPERTY(EditAnywhere, BlueprintReadWrite, Getter = "Auto", BlueprintSetter = SetBackMaterial, Category = "Text3D")
-	TObjectPtr<class UMaterialInterface> BackMaterial;
-
-	/** Text font */
-	UPROPERTY(EditAnywhere, BlueprintReadWrite, BlueprintSetter = SetFont, Category = "Text3D")
-	TObjectPtr<class UFont> Font;
-
-	/** Horizontal text alignment */
-	UPROPERTY(EditAnywhere, BlueprintReadWrite, BlueprintSetter = SetHorizontalAlignment, Category = "Text3D")
-	EText3DHorizontalTextAlignment HorizontalAlignment;
-
-	/** Vertical text alignment */
-	UPROPERTY(EditAnywhere, BlueprintReadWrite, BlueprintSetter = SetVerticalAlignment, Category = "Text3D")
-	EText3DVerticalTextAlignment VerticalAlignment;
-
-	/** Text kerning */
-	UPROPERTY(EditAnywhere, BlueprintReadWrite, BlueprintSetter = SetKerning, Category = "Text3D")
-	float Kerning;
-
-	/** Extra line spacing */
-	UPROPERTY(EditAnywhere, BlueprintReadWrite, BlueprintSetter = SetLineSpacing, Category = "Text3D")
-	float LineSpacing;
-
-	/** Extra word spacing */
-	UPROPERTY(EditAnywhere, BlueprintReadWrite, BlueprintSetter = SetWordSpacing, Category = "Text3D")
-	float WordSpacing;
-
-	/** Enables a maximum width to the 3D Text */
-	UPROPERTY(EditAnywhere, BlueprintReadWrite, BlueprintSetter = SetHasMaxWidth, Category = "Text3D", meta = (InlineEditConditionToggle))
-	bool bHasMaxWidth;
-
-	/** Sets a maximum width to the 3D Text */
-	UPROPERTY(EditAnywhere, BlueprintReadWrite, BlueprintSetter = SetMaxWidth, Category = "Text3D", meta = (EditCondition = "bHasMaxWidth", ClampMin = 1))
-	float MaxWidth;
-
-	/** Enables a maximum height to the 3D Text */
-	UPROPERTY(EditAnywhere, BlueprintReadWrite, BlueprintSetter = SetHasMaxHeight, Category = "Text3D", meta = (InlineEditConditionToggle))
-	bool bHasMaxHeight;
-
-	/** Sets a maximum height to the 3D Text */
-	UPROPERTY(EditAnywhere, BlueprintReadWrite, BlueprintSetter = SetMaxHeight, Category = "Text3D", meta = (EditCondition = "bHasMaxHeight", ClampMin = 1))
-	float MaxHeight;
-
-	/** Should the mesh scale proportionally when Max Width/Height is set */
-	UPROPERTY(EditAnywhere, BlueprintReadWrite, BlueprintSetter = SetScaleProportionally, Category = "Text3D")
-	bool bScaleProportionally;
-
-	// Lighting flags
-	
-	/** Controls whether the text glyphs should cast a shadow or not. */
-	UPROPERTY(EditAnywhere, BlueprintReadWrite, BlueprintSetter="SetCastShadow", Category=Lighting)
-	bool bCastShadow = true;
-
 	/**
 	 * Delegate called after text is rebuilt
 	 */
@@ -169,103 +73,178 @@ public:
 	DECLARE_MULTICAST_DELEGATE(FTextGeneratedNative);
 	FTextGeneratedNative& OnTextGenerated() { return TextGeneratedNativeDelegate; }
 
+	/** Get whether to allow automatic refresh/mesh generation */
+	bool RefreshesOnChange() const;
+
 	/** Set whether to allow automatic refresh/mesh generation */
-	UFUNCTION(BlueprintCallable, Category = "Rendering|Components|Text3D")
-	void SetRefreshOnChange(const bool& Value);
+	void SetRefreshOnChange(const bool Value);
+
+	/** Get the text value and signal the primitives to be rebuilt */
+	const FText& GetText() const;
 
 	/** Set the text value and signal the primitives to be rebuilt */
-	UFUNCTION(BlueprintCallable, Category = "Rendering|Components|Text3D")
+	UFUNCTION(BlueprintCallable, Category = "Rendering|Components|Text3D", meta = (DeprecatedFunction, DeprecationMessage = "Set the property directly"))
 	void SetText(const FText& Value);
 
+	/** Get the text font and signal the primitives to be rebuilt */
+	const UFont* GetFont() const;
+
 	/** Set the text font and signal the primitives to be rebuilt */
-	UFUNCTION(BlueprintCallable, Category = "Rendering|Components|Text3D")
-	void SetFont(class UFont* const InFont);
+	UFUNCTION(BlueprintCallable, Category = "Rendering|Components|Text3D", meta = (DeprecatedFunction, DeprecationMessage = "Set the property directly"))
+	void SetFont(UFont* InFont);
 
-	UFUNCTION(BlueprintCallable, Category = "Rendering|Components|Text3D")
-	void SetOutline(const bool bValue);
-	
-	UFUNCTION(BlueprintCallable, Category = "Rendering|Components|Text3D")
+	/** Get whether an outline is applied. */
+	bool HasOutline() const;
+
+	/** Set whether an outline is applied. */
+	UFUNCTION(BlueprintCallable, Category = "Rendering|Components|Text3D", meta = (DeprecatedFunction, DeprecationMessage = "Set the property directly"))
+	void SetHasOutline(const bool bValue);
+
+	/** Get the outline width. */
+	float GetOutlineExpand() const;
+
+	/** Set the outline width. */
+	UFUNCTION(BlueprintCallable, Category = "Rendering|Components|Text3D", meta = (DeprecatedFunction, DeprecationMessage = "Set the property directly"))
 	void SetOutlineExpand(const float Value);
 
+	/** Get the text extrusion size and signal the primitives to be rebuilt */
+	float GetExtrude() const;
+
 	/** Set the text extrusion size and signal the primitives to be rebuilt */
-	UFUNCTION(BlueprintCallable, Category = "Rendering|Components|Text3D")
+	UFUNCTION(BlueprintCallable, Category = "Rendering|Components|Text3D", meta = (DeprecatedFunction, DeprecationMessage = "Set the property directly"))
 	void SetExtrude(const float Value);
 
+	/** Get the 3d bevel value */
+	float GetBevel() const;
+
 	/** Set the 3d bevel value */
-	UFUNCTION(BlueprintCallable, Category = "Rendering|Components|Text3D")
+	UFUNCTION(BlueprintCallable, Category = "Rendering|Components|Text3D", meta = (DeprecatedFunction, DeprecationMessage = "Set the property directly"))
 	void SetBevel(const float Value);
 
+	/** Get the 3d bevel type */
+	EText3DBevelType GetBevelType() const;
+
 	/** Set the 3d bevel type */
-	UFUNCTION(BlueprintCallable, Category = "Rendering|Components|Text3D")
+	UFUNCTION(BlueprintCallable, Category = "Rendering|Components|Text3D", meta = (DeprecatedFunction, DeprecationMessage = "Set the property directly"))
 	void SetBevelType(const EText3DBevelType Value);
 
-	/** Set the amount of segments that will be used to tesselate the Bevel */
-	UFUNCTION(BlueprintCallable, Category = "Rendering|Components|Text3D")
+	/** Get the amount of segments that will be used to tessellate the Bevel */
+	int32 GetBevelSegments() const;
+
+	/** Set the amount of segments that will be used to tessellate the Bevel */
+	UFUNCTION(BlueprintCallable, Category = "Rendering|Components|Text3D", meta = (DeprecatedFunction, DeprecationMessage = "Set the property directly"))
 	void SetBevelSegments(const int32 Value);
 
+	/** Get the text front material */
+	UMaterialInterface* GetFrontMaterial() const;
 
 	/** Set the text front material */
-	UFUNCTION(BlueprintCallable, Category = "Rendering|Components|Text3D")
-	void SetFrontMaterial(class UMaterialInterface* Value);
+	UFUNCTION(BlueprintCallable, Category = "Rendering|Components|Text3D", meta = (DeprecatedFunction, DeprecationMessage = "Set the property directly"))
+	void SetFrontMaterial(UMaterialInterface* Value);
+
+	/** Get the text bevel material */
+	UMaterialInterface* GetBevelMaterial() const;
 
 	/** Set the text bevel material */
-	UFUNCTION(BlueprintCallable, Category = "Rendering|Components|Text3D")
-	void SetBevelMaterial(class UMaterialInterface* Value);
+	UFUNCTION(BlueprintCallable, Category = "Rendering|Components|Text3D", meta = (DeprecatedFunction, DeprecationMessage = "Set the property directly"))
+	void SetBevelMaterial(UMaterialInterface* Value);
+
+	/** Get the text extrude material */
+	UMaterialInterface* GetExtrudeMaterial() const;
 
 	/** Set the text extrude material */
-	UFUNCTION(BlueprintCallable, Category = "Rendering|Components|Text3D")
-	void SetExtrudeMaterial(class UMaterialInterface* Value);
+	UFUNCTION(BlueprintCallable, Category = "Rendering|Components|Text3D", meta = (DeprecatedFunction, DeprecationMessage = "Set the property directly"))
+	void SetExtrudeMaterial(UMaterialInterface* Value);
+
+	/** Get the text back material */
+	UMaterialInterface* GetBackMaterial() const;
 
 	/** Set the text back material */
-	UFUNCTION(BlueprintCallable, Category = "Rendering|Components|Text3D")
-	void SetBackMaterial(class UMaterialInterface* Value);
+	UFUNCTION(BlueprintCallable, Category = "Rendering|Components|Text3D", meta = (DeprecatedFunction, DeprecationMessage = "Set the property directly"))
+	void SetBackMaterial(UMaterialInterface* Value);
+
+	/** Get the kerning value and signal the primitives to be rebuilt */
+	float GetKerning() const;
 
 	/** Set the kerning value and signal the primitives to be rebuilt */
-	UFUNCTION(BlueprintCallable, Category = "Rendering|Components|Text3D")
+	UFUNCTION(BlueprintCallable, Category = "Rendering|Components|Text3D", meta = (DeprecatedFunction, DeprecationMessage = "Set the property directly"))
 	void SetKerning(const float Value);
 
+	/** Get the line spacing value and signal the primitives to be rebuilt */
+	float GetLineSpacing() const;
+
 	/** Set the line spacing value and signal the primitives to be rebuilt */
-	UFUNCTION(BlueprintCallable, Category = "Rendering|Components|Text3D")
+	UFUNCTION(BlueprintCallable, Category = "Rendering|Components|Text3D", meta = (DeprecatedFunction, DeprecationMessage = "Set the property directly"))
 	void SetLineSpacing(const float Value);
 
+	/** Get the word spacing value and signal the primitives to be rebuilt */
+	float GetWordSpacing() const;
+
 	/** Set the word spacing value and signal the primitives to be rebuilt */
-	UFUNCTION(BlueprintCallable, Category = "Rendering|Components|Text3D")
+	UFUNCTION(BlueprintCallable, Category = "Rendering|Components|Text3D", meta = (DeprecatedFunction, DeprecationMessage = "Set the property directly"))
 	void SetWordSpacing(const float Value);
 
+	/** Get the horizontal alignment value and signal the primitives to be rebuilt */
+	EText3DHorizontalTextAlignment GetHorizontalAlignment() const;
+
 	/** Set the horizontal alignment value and signal the primitives to be rebuilt */
-	UFUNCTION(BlueprintCallable, Category = "Rendering|Components|Text3D")
+	UFUNCTION(BlueprintCallable, Category = "Rendering|Components|Text3D", meta = (DeprecatedFunction, DeprecationMessage = "Set the property directly"))
 	void SetHorizontalAlignment(const EText3DHorizontalTextAlignment value);
 
+	/** Get the vertical alignment and signal the primitives to be rebuilt */
+	EText3DVerticalTextAlignment GetVerticalAlignment() const;
+
 	/** Set the vertical alignment and signal the primitives to be rebuilt */
-	UFUNCTION(BlueprintCallable, Category = "Rendering|Components|Text3D")
+	UFUNCTION(BlueprintCallable, Category = "Rendering|Components|Text3D", meta = (DeprecatedFunction, DeprecationMessage = "Set the property directly"))
 	void SetVerticalAlignment(const EText3DVerticalTextAlignment value);
 
+	/** Whether a maximum width is specified */
+	bool HasMaxWidth() const;
+
 	/** Enable / Disable a Maximum Width */
-	UFUNCTION(BlueprintCallable, Category = "Rendering|Components|Text3D")
+	UFUNCTION(BlueprintCallable, Category = "Rendering|Components|Text3D", meta = (DeprecatedFunction, DeprecationMessage = "Set the property directly"))
 	void SetHasMaxWidth(const bool Value);
 
+	/** Get the Maximum Width - If width is larger, mesh will scale down to fit MaxWidth value */
+	float GetMaxWidth() const;
+
 	/** Set the Maximum Width - If width is larger, mesh will scale down to fit MaxWidth value */
-	UFUNCTION(BlueprintCallable, Category = "Rendering|Components|Text3D")
+	UFUNCTION(BlueprintCallable, Category = "Rendering|Components|Text3D", meta = (DeprecatedFunction, DeprecationMessage = "Set the property directly"))
 	void SetMaxWidth(const float Value);
 
+	/** Whether a maximum height is specified */
+	bool HasMaxHeight() const;
+
 	/** Enable / Disable a Maximum Height */
-	UFUNCTION(BlueprintCallable, Category = "Rendering|Components|Text3D")
+	UFUNCTION(BlueprintCallable, Category = "Rendering|Components|Text3D", meta = (DeprecatedFunction, DeprecationMessage = "Set the property directly"))
 	void SetHasMaxHeight(const bool Value);
 
+	/** Get the Maximum Height - If height is larger, mesh will scale down to fit MaxHeight value */
+	float GetMaxHeight() const;
+
 	/** Set the Maximum Height - If height is larger, mesh will scale down to fit MaxHeight value */
-	UFUNCTION(BlueprintCallable, Category = "Rendering|Components|Text3D")
+	UFUNCTION(BlueprintCallable, Category = "Rendering|Components|Text3D", meta = (DeprecatedFunction, DeprecationMessage = "Set the property directly"))
 	void SetMaxHeight(const float Value);
 
+	/** Get if the mesh should scale proportionally when Max Width/Height is set */
+	bool ScalesProportionally() const;
+
 	/** Set if the mesh should scale proportionally when Max Width/Height is set */
-	UFUNCTION(BlueprintCallable, Category = "Rendering|Components|Text3D")
+	UFUNCTION(BlueprintCallable, Category = "Rendering|Components|Text3D", meta = (DeprecatedFunction, DeprecationMessage = "Set the property directly"))
 	void SetScaleProportionally(const bool Value);
 
 	/** Freeze mesh rebuild, to avoid unnecessary mesh rebuilds when setting a few properties together */
-	UFUNCTION(BlueprintCallable, Category = "Rendering|Components|Text3D")
+	bool IsFrozen() const;
+
+	/** Freeze mesh rebuild, to avoid unnecessary mesh rebuilds when setting a few properties together */
+	UFUNCTION(BlueprintCallable, Category = "Rendering|Components|Text3D", meta = (DeprecatedFunction, DeprecationMessage = "Set the property directly"))
 	void SetFreeze(const bool bFreeze);
 
-	/** Changes the value of CastShadow. */
-	UFUNCTION(BlueprintCallable, Category="Rendering")
+	/** Get the value of CastShadow. */
+	bool CastsShadow() const;
+
+	/** Set the value of CastShadow. */
+	UFUNCTION(BlueprintCallable, Category = "Rendering|Components|Text3D", meta = (DeprecatedFunction, DeprecationMessage = "Set the property directly"))
 	void SetCastShadow(bool NewCastShadow);
 
 	UFUNCTION(BlueprintCallable, Category = "Rendering|Components|Text3D")
@@ -299,14 +278,117 @@ public:
 	void Rebuild();
 
 protected:
+	/** Whether to allow automatic refresh/mesh generation */
+	UPROPERTY(EditAnywhere, BlueprintReadWrite, Getter = "RefreshesOnChange", Setter = "SetRefreshOnChange", Category = "Text3D", AdvancedDisplay, meta = (AllowPrivateAccess))
+	bool bRefreshOnChange = true;
+
+	/** The text to generate a 3d mesh */
+	UPROPERTY(EditAnywhere, BlueprintReadWrite, Getter, Setter, Category = "Text3D", meta = (MultiLine = true, AllowPrivateAccess))
+	FText Text;
+
+	/** Size of the extrude */
+	UPROPERTY(EditAnywhere, BlueprintReadWrite, Getter, Setter, Category = "Text3D", meta = (ClampMin = 0, AllowPrivateAccess))
+	float Extrude;
+
+	/** Size of bevel */
+	UPROPERTY(EditAnywhere, BlueprintReadWrite, Getter, Setter, Category = "Text3D", meta = (EditCondition = "!bOutline", ClampMin = 0, AllowPrivateAccess))
+	float Bevel;
+
+	/** Bevel Type */
+	UPROPERTY(EditAnywhere, BlueprintReadWrite, Getter, Setter, Category = "Text3D", meta = (EditCondition = "!bOutline", AllowPrivateAccess))
+	EText3DBevelType BevelType;
+
+	/** Bevel Segments (Defines the amount of tesselation for the bevel part) */
+	UPROPERTY(EditAnywhere, BlueprintReadWrite, Getter, Setter, Category = "Text3D", meta = (EditCondition = "!bOutline", ClampMin = 1, ClampMax = 15, AllowPrivateAccess))
+	int32 BevelSegments;
+
+	/** Generate Outline */
+	UPROPERTY(EditAnywhere, BlueprintReadWrite, Getter = "HasOutline", Setter = "SetHasOutline", Category = "Text3D", meta = (AllowPrivateAccess))
+	bool bOutline;
+
+	/** Outline expand/offset amount */
+	UPROPERTY(EditAnywhere, BlueprintReadWrite, Getter, Setter, Category = "Text3D", meta = (EditCondition = "bOutline", AllowPrivateAccess))
+	float OutlineExpand;
+
+	/** Material for the front part */
+	UPROPERTY(EditAnywhere, BlueprintReadWrite, Getter, Setter, Category = "Text3D", meta = (AllowPrivateAccess))
+	TObjectPtr<UMaterialInterface> FrontMaterial;
+
+	/** Material for the bevel part */
+	UPROPERTY(EditAnywhere, BlueprintReadWrite, Getter, Setter, Category = "Text3D", meta = (EditCondition = "!bOutline", AllowPrivateAccess))
+	TObjectPtr<UMaterialInterface> BevelMaterial;
+
+	/** Material for the extruded part */
+	UPROPERTY(EditAnywhere, BlueprintReadWrite, Getter, Setter, Category = "Text3D", meta = (AllowPrivateAccess))
+	TObjectPtr<UMaterialInterface> ExtrudeMaterial;
+
+	/** Material for the back part */
+	UPROPERTY(EditAnywhere, BlueprintReadWrite, Getter, Setter, Category = "Text3D", meta = (AllowPrivateAccess))
+	TObjectPtr<UMaterialInterface> BackMaterial;
+
+	/** Text font */
+	UPROPERTY(EditAnywhere, BlueprintReadWrite, Getter, Setter, Category = "Text3D", meta = (AllowPrivateAccess))
+	TObjectPtr<UFont> Font;
+
+	/** Horizontal text alignment */
+	UPROPERTY(EditAnywhere, BlueprintReadWrite, Getter, Setter, Category = "Text3D", meta = (AllowPrivateAccess))
+	EText3DHorizontalTextAlignment HorizontalAlignment;
+
+	/** Vertical text alignment */
+	UPROPERTY(EditAnywhere, BlueprintReadWrite, Getter, Setter, Category = "Text3D", meta = (AllowPrivateAccess))
+	EText3DVerticalTextAlignment VerticalAlignment;
+
+	/** Text kerning */
+	UPROPERTY(EditAnywhere, BlueprintReadWrite, Getter, Setter, Category = "Text3D", meta = (AllowPrivateAccess))
+	float Kerning;
+
+	/** Extra line spacing */
+	UPROPERTY(EditAnywhere, BlueprintReadWrite, Getter, Setter, Category = "Text3D", meta = (AllowPrivateAccess))
+	float LineSpacing;
+
+	/** Extra word spacing */
+	UPROPERTY(EditAnywhere, BlueprintReadWrite, Setter, Category = "Text3D", meta = (AllowPrivateAccess))
+	float WordSpacing;
+
+	/** Enables a maximum width to the 3D Text */
+	UPROPERTY(EditAnywhere, BlueprintReadWrite, Getter = "HasMaxWidth", Setter = "SetHasMaxWidth", Category = "Text3D", meta = (InlineEditConditionToggle, AllowPrivateAccess))
+	bool bHasMaxWidth;
+
+	/** Sets a maximum width to the 3D Text */
+	UPROPERTY(EditAnywhere, BlueprintReadWrite, Getter, Setter, Category = "Text3D", meta = (EditCondition = "bHasMaxWidth", ClampMin = 1, AllowPrivateAccess))
+	float MaxWidth;
+
+	/** Enables a maximum height to the 3D Text */
+	UPROPERTY(EditAnywhere, BlueprintReadWrite, Getter = "HasMaxHeight", Setter = "SetHasMaxHeight", Category = "Text3D", meta = (InlineEditConditionToggle, AllowPrivateAccess))
+	bool bHasMaxHeight;
+
+	/** Sets a maximum height to the 3D Text */
+	UPROPERTY(EditAnywhere, BlueprintReadWrite, Getter, Setter, Category = "Text3D", meta = (EditCondition = "bHasMaxHeight", ClampMin = 1, AllowPrivateAccess))
+	float MaxHeight;
+
+	/** Should the mesh scale proportionally when Max Width/Height is set */
+	UPROPERTY(EditAnywhere, BlueprintReadWrite, Getter = "ScalesProportionally", Setter = "SetScaleProportionally", Category = "Text3D", meta = (AllowPrivateAccess))
+	bool bScaleProportionally;
+
+	// Lighting flags
+	
+	/** Controls whether the text glyphs should cast a shadow or not. */
+	UPROPERTY(EditAnywhere, BlueprintReadWrite, Getter = "CastsShadow", Setter = "SetCastShadow", Category = "Lighting", meta = (AllowPrivateAccess))
+	bool bCastShadow = true;
+
+protected:
+	/** Intercept and propagate a change on this component to all children. */
 	virtual void OnVisibilityChanged() override;
+
+	/** Intercept and propagate a change on this component to all children. */
 	virtual void OnHiddenInGameChanged() override;
 
+	/** Clears all generated components and meshes from this component. */
 	void ClearTextMesh();
 
 private:
 	UPROPERTY()
-	TObjectPtr<class USceneComponent> TextRoot;
+	TObjectPtr<USceneComponent> TextRoot;
 
 	UPROPERTY(BlueprintAssignable, Category = Events, meta = (AllowPrivateAccess = true, DisplayName = "On Text Generated"))
 	FTextGenerated TextGeneratedDelegate;
@@ -315,44 +397,61 @@ private:
 
 	/** Flagged as true when the text mesh is being built. */
 	std::atomic<bool> bIsBuilding;
-	
+
+	/** No mesh or layout rebuilds occur while this flag is true, allowing the user to set numerous properties before rebuilding. */
 	bool bFreezeBuild;
+
+	/** Used to determine and selectively perform the the type of rebuild requested. */
 	EText3DModifyFlags ModifyFlags;
 
+	/** Additional scale to apply to the text. */
 	FVector TextScale;
 
-	TSharedPtr<struct FText3DShapedText> ShapedText;
+	/** To uniquely identify and discard duplicate update requests. */
 	TArray<TSharedPtr<int32>> CachedCounterReferences;
 
+	/** Caches the last result of ShapedText, to allow faster updates of layout changes. */
+	TSharedPtr<struct FText3DShapedText> ShapedText;
+
+	/** Stores the text layout calculated by the TextLayoutMarshaller. */
 	TSharedPtr<FTextLayout> TextLayout;
+
+	/** Determines how text is laid out, ie. parsing line breaks. */
 	TSharedPtr<ITextLayoutMarshaller> TextLayoutMarshaller;
 
+	/** Each character mesh is parented to a Kerning component. */
 	UPROPERTY(Transient)
 	TArray<TObjectPtr<USceneComponent>> CharacterKernings;
 
+	/** Each character mesh is held in a component. */
 	UPROPERTY(Transient)
 	TArray<TObjectPtr<UStaticMeshComponent>> CharacterMeshes;
 
 	/** Allocates, or shrinks existing components to match the input number. Returns false if nothing modified. */
 	bool AllocateGlyphs(int32 Num);
 
-	class UMaterialInterface* GetMaterial(const EText3DGroupType Type) const;
-	void SetMaterial(const EText3DGroupType Type, class UMaterialInterface* Material);
+	/** Slot based Material accessors. */
+	UMaterialInterface* GetMaterial(const EText3DGroupType Type) const;
+	void SetMaterial(const EText3DGroupType Type, UMaterialInterface* Material);
+	void UpdateMaterial(const EText3DGroupType Type, UMaterialInterface* Material);
 
+	/** Validation and context sensitive limits applied to the current bevel value. */
+	void CheckBevel();
+	float MaxBevel() const;
+
+	/** Selective build functionality (geometry, only layout, etc.) */
 	void UpdateTransforms();
 	void RebuildInternal(const bool& bIsAutoUpdate = true, const bool& bCleanCache = false);
 	void BuildTextMesh(const bool& bCleanCache = false);
 	void BuildTextMeshInternal(const bool& bCleanCache);
-	void CheckBevel();
-	float MaxBevel() const;
-
-	void UpdateMaterial(const EText3DGroupType Type, class UMaterialInterface* Material);
 
+	/** Layout functionality. */
 	void CalculateTextWidth();
 	float GetTextHeight() const;
 	void CalculateTextScale();
 	FVector GetLineLocation(int32 LineIndex);
-	
+
+	/** Convenience functions to set and query ModifyFlags. */
 	bool NeedsMeshRebuild() const;
 	bool NeedsLayoutUpdate() const;
 	void MarkForGeometryUpdate();
