[Text3D] Material assignment crash when using component within custom actor

diff --git a/Engine/Plugins/Experimental/Text3D/Source/Text3D/Private/Text3DComponent.cpp b/Engine/Plugins/Experimental/Text3D/Source/Text3D/Private/Text3DComponent.cpp
--- a/Engine/Plugins/Experimental/Text3D/Source/Text3D/Private/Text3DComponent.cpp
+++ b/Engine/Plugins/Experimental/Text3D/Source/Text3D/Private/Text3DComponent.cpp
@@ -24,6 +24,30 @@
 
 #define LOCTEXT_NAMESPACE "Text3D"
 
+#if WITH_EDITOR
+namespace UE::Text3D::Private
+{
+	/** Get the Group Type according to the material name. */
+	static TMap<FName, EText3DGroupType> MaterialToGroup =
+	{
+		{ GET_MEMBER_NAME_CHECKED(UText3DComponent, FrontMaterial), EText3DGroupType::Front },
+		{ GET_MEMBER_NAME_CHECKED(UText3DComponent, BackMaterial), EText3DGroupType::Back },
+		{ GET_MEMBER_NAME_CHECKED(UText3DComponent, ExtrudeMaterial), EText3DGroupType::Extrude },
+		{ GET_MEMBER_NAME_CHECKED(UText3DComponent, BevelMaterial), EText3DGroupType::Bevel }
+	};
+
+	/** Sets the material based on the group name. */
+	using FGetter = TMemFunPtrType<true, UText3DComponent, UMaterialInterface*()>::Type;
+	static TMap<FName, FGetter> GroupToMaterial =
+	{
+		{ GET_MEMBER_NAME_CHECKED(UText3DComponent, FrontMaterial), &UText3DComponent::GetFrontMaterial },
+		{ GET_MEMBER_NAME_CHECKED(UText3DComponent, BackMaterial), &UText3DComponent::GetBackMaterial },
+		{ GET_MEMBER_NAME_CHECKED(UText3DComponent, ExtrudeMaterial), &UText3DComponent::GetExtrudeMaterial },
+		{ GET_MEMBER_NAME_CHECKED(UText3DComponent, BevelMaterial), &UText3DComponent::GetBevelMaterial }
+	};
+}
+#endif
+
 struct FText3DShapedText
 {
 	FText3DShapedText()
@@ -169,23 +193,6 @@ void UText3DComponent::PostEditChangeProperty(FPropertyChangedEvent& PropertyCha
 {
 	Super::PostEditChangeProperty(PropertyChangedEvent);
 
-	static TMap<FName, EText3DGroupType> MaterialToGroup =
-	{
-		{ GET_MEMBER_NAME_CHECKED(UText3DComponent, FrontMaterial), EText3DGroupType::Front },
-		{ GET_MEMBER_NAME_CHECKED(UText3DComponent, BackMaterial), EText3DGroupType::Back },
-		{ GET_MEMBER_NAME_CHECKED(UText3DComponent, ExtrudeMaterial), EText3DGroupType::Extrude },
-		{ GET_MEMBER_NAME_CHECKED(UText3DComponent, BevelMaterial), EText3DGroupType::Bevel }
-	};
-
-	using FGetter = TDelegate<UMaterialInterface*()>;
-	static TMap<FName, FGetter> MaterialLookup =
-	{
-		{ GET_MEMBER_NAME_CHECKED(UText3DComponent, FrontMaterial), FGetter::CreateUObject(this, &UText3DComponent::GetFrontMaterial) },
-		{ GET_MEMBER_NAME_CHECKED(UText3DComponent, BackMaterial), FGetter::CreateUObject(this, &UText3DComponent::GetBackMaterial) },
-		{ GET_MEMBER_NAME_CHECKED(UText3DComponent, ExtrudeMaterial), FGetter::CreateUObject(this, &UText3DComponent::GetExtrudeMaterial) },
-		{ GET_MEMBER_NAME_CHECKED(UText3DComponent, BevelMaterial), FGetter::CreateUObject(this, &UText3DComponent::GetBevelMaterial) }
-	};
-
 	const FName Name = PropertyChangedEvent.GetPropertyName();
 
 	static FName BevelTypePropertyName = GET_MEMBER_NAME_CHECKED(UText3DComponent, BevelType);
@@ -244,9 +251,9 @@ void UText3DComponent::PostEditChangeProperty(FPropertyChangedEvent& PropertyCha
 	{
 		MarkForLayoutUpdate();
 	}
-	else if (const EText3DGroupType* MaterialGroup = MaterialToGroup.Find(Name))
+	else if (const EText3DGroupType* MaterialGroup = UE::Text3D::Private::MaterialToGroup.Find(Name))
 	{
-		UpdateMaterial(*MaterialGroup, MaterialLookup[Name].Execute());
+		UpdateMaterial(*MaterialGroup, Invoke(UE::Text3D::Private::GroupToMaterial[Name], this));
 	}
 	// Any property not explicitly handled should trigger a full rebuild
 	else
