[Text3D] Match kerning behavior to slate/umg. Will affect existing assets, but previous change introduced prior to public release.

diff --git a/Engine/Plugins/Experimental/Text3D/Source/Text3D/Private/Text3DComponent.cpp b/Engine/Plugins/Experimental/Text3D/Source/Text3D/Private/Text3DComponent.cpp
--- a/Engine/Plugins/Experimental/Text3D/Source/Text3D/Private/Text3DComponent.cpp
+++ b/Engine/Plugins/Experimental/Text3D/Source/Text3D/Private/Text3DComponent.cpp
@@ -1148,12 +1148,12 @@ void UText3DComponent::BuildTextMeshInternal(const bool& bCleanCache)
 	for (int32 LineIndex = 0; LineIndex < ShapedText->Lines.Num(); LineIndex++)
 	{
 		const FShapedGlyphLine& ShapedLine = ShapedText->Lines[LineIndex];
-		FVector Location = GetLineLocation(LineIndex);
+		FVector LineLocation = GetLineLocation(LineIndex);
 
 		for (int32 LineGlyph = 0; LineGlyph < ShapedLine.GlyphsToRender.Num(); LineGlyph++)
 		{
-			FVector GlyphLocation = Location;
-			Location.Y += ShapedLine.GetAdvance(LineGlyph, Kerning, WordSpacing);
+			FVector GlyphLocation = LineLocation;
+			LineLocation.Y += ShapedLine.GetAdvance(LineGlyph, Kerning, WordSpacing);
 
 			const FShapedGlyphEntry& ShapedGlyph = ShapedLine.GlyphsToRender[LineGlyph];
 			if (!ShapedGlyph.bIsVisible)
diff --git a/Engine/Plugins/Experimental/Text3D/Source/Text3D/Private/TextShaper.h b/Engine/Plugins/Experimental/Text3D/Source/Text3D/Private/TextShaper.h
--- a/Engine/Plugins/Experimental/Text3D/Source/Text3D/Private/TextShaper.h
+++ b/Engine/Plugins/Experimental/Text3D/Source/Text3D/Private/TextShaper.h
@@ -31,7 +31,8 @@ struct FShapedGlyphLine
 
 		if (Index < GlyphsToRender.Num() - 1)
 		{
-			Advance += Glyph.Kerning + Kerning;
+			// @note: as per FSlateElementBatcher::BuildShapedTextSequence, per Glyph Kerning isn't used
+			Advance += Kerning;
 
 			if (!Glyph.bIsVisible)
 			{
