[Text3D] cis issue

diff --git a/Engine/Plugins/Experimental/Text3D/Source/Text3D/Private/Text3DComponent.cpp b/Engine/Plugins/Experimental/Text3D/Source/Text3D/Private/Text3DComponent.cpp
--- a/Engine/Plugins/Experimental/Text3D/Source/Text3D/Private/Text3DComponent.cpp
+++ b/Engine/Plugins/Experimental/Text3D/Source/Text3D/Private/Text3DComponent.cpp
@@ -30,20 +30,20 @@ namespace UE::Text3D::Private
 	/** Get the Group Type according to the material name. */
 	static TMap<FName, EText3DGroupType> MaterialToGroup =
 	{
-		{ GET_MEMBER_NAME_CHECKED(UText3DComponent, FrontMaterial), EText3DGroupType::Front },
-		{ GET_MEMBER_NAME_CHECKED(UText3DComponent, BackMaterial), EText3DGroupType::Back },
-		{ GET_MEMBER_NAME_CHECKED(UText3DComponent, ExtrudeMaterial), EText3DGroupType::Extrude },
-		{ GET_MEMBER_NAME_CHECKED(UText3DComponent, BevelMaterial), EText3DGroupType::Bevel }
+		{ TEXT("FrontMaterial"), EText3DGroupType::Front },
+		{ TEXT("BackMaterial"), EText3DGroupType::Back },
+		{ TEXT("ExtrudeMaterial"), EText3DGroupType::Extrude },
+		{ TEXT("BevelMaterial"), EText3DGroupType::Bevel }
 	};
 
 	/** Sets the material based on the group name. */
 	using FGetter = TMemFunPtrType<true, UText3DComponent, UMaterialInterface*()>::Type;
 	static TMap<FName, FGetter> GroupToMaterial =
 	{
-		{ GET_MEMBER_NAME_CHECKED(UText3DComponent, FrontMaterial), &UText3DComponent::GetFrontMaterial },
-		{ GET_MEMBER_NAME_CHECKED(UText3DComponent, BackMaterial), &UText3DComponent::GetBackMaterial },
-		{ GET_MEMBER_NAME_CHECKED(UText3DComponent, ExtrudeMaterial), &UText3DComponent::GetExtrudeMaterial },
-		{ GET_MEMBER_NAME_CHECKED(UText3DComponent, BevelMaterial), &UText3DComponent::GetBevelMaterial }
+		{ TEXT("FrontMaterial"), &UText3DComponent::GetFrontMaterial },
+		{ TEXT("BackMaterial"), &UText3DComponent::GetBackMaterial },
+		{ TEXT("ExtrudeMaterial"), &UText3DComponent::GetExtrudeMaterial },
+		{ TEXT("BevelMaterial"), &UText3DComponent::GetBevelMaterial }
 	};
 }
 #endif
