[Text3D] Avoid marking anything from update on Transform Post Edit Change

diff --git a/Engine/Plugins/Experimental/Text3D/Source/Text3D/Private/Text3DComponent.cpp b/Engine/Plugins/Experimental/Text3D/Source/Text3D/Private/Text3DComponent.cpp
--- a/Engine/Plugins/Experimental/Text3D/Source/Text3D/Private/Text3DComponent.cpp
+++ b/Engine/Plugins/Experimental/Text3D/Source/Text3D/Private/Text3DComponent.cpp
@@ -251,6 +251,13 @@ void UText3DComponent::PostEditChangeProperty(FPropertyChangedEvent& PropertyCha
 	{
 		MarkForLayoutUpdate();
 	}
+	else if (Name == USceneComponent::GetRelativeLocationPropertyName() ||
+			 Name == USceneComponent::GetRelativeRotationPropertyName() ||
+			 Name == USceneComponent::GetRelativeScale3DPropertyName())
+	{
+		// Nothing should be marked for update if Transform was changed, so use the parent of UText3DComponent PECP
+		UText3DComponent::Super::PostEditChangeProperty(PropertyChangedEvent);
+	}
 	else if (const EText3DGroupType* MaterialGroup = UE::Text3D::Private::MaterialToGroup.Find(Name))
 	{
 		UpdateMaterial(*MaterialGroup, Invoke(UE::Text3D::Private::GroupToMaterial[Name], this));
