[Text3D] PECP triggering unneeded geometry and layout rebuild

- Removed redundant Super::PECP call for Actor Component Transform properties (already called at the top of PECP)
- Removed Geometry and Layout updates triggered for properties not handled directly by PECP. This was causing unneeded geometry rebuilding.

diff --git a/Engine/Plugins/Experimental/Text3D/Source/Text3D/Private/Text3DComponent.cpp b/Engine/Plugins/Experimental/Text3D/Source/Text3D/Private/Text3DComponent.cpp
--- a/Engine/Plugins/Experimental/Text3D/Source/Text3D/Private/Text3DComponent.cpp
+++ b/Engine/Plugins/Experimental/Text3D/Source/Text3D/Private/Text3DComponent.cpp
@@ -251,23 +251,10 @@ void UText3DComponent::PostEditChangeProperty(FPropertyChangedEvent& PropertyCha
 	{
 		MarkForLayoutUpdate();
 	}
-	else if (Name == USceneComponent::GetRelativeLocationPropertyName() ||
-			 Name == USceneComponent::GetRelativeRotationPropertyName() ||
-			 Name == USceneComponent::GetRelativeScale3DPropertyName())
-	{
-		// Nothing should be marked for update if Transform was changed, so use the parent of UText3DComponent PECP
-		UText3DComponent::Super::PostEditChangeProperty(PropertyChangedEvent);
-	}
 	else if (const EText3DGroupType* MaterialGroup = UE::Text3D::Private::MaterialToGroup.Find(Name))
 	{
 		UpdateMaterial(*MaterialGroup, Invoke(UE::Text3D::Private::GroupToMaterial[Name], this));
 	}
-	// Any property not explicitly handled should trigger a full rebuild
-	else
-	{
-		MarkForGeometryUpdate();
-		MarkForLayoutUpdate();
-	}
 
 	RebuildInternal();
 }
