[Text3D] Add virtual call to format text

Adds:

- FormatText virtual function: allows to modify text right before text geometry is generated. Will not modify the Text property.

- GetFormattedText function: will return a copy of the text after the modifications performed by FormatText

- TriggerInternalRebuild: can be used to force a rebuild of layout and/or geometry.

diff --git a/Engine/Plugins/Experimental/Text3D/Source/Text3D/Private/Text3DComponent.cpp b/Engine/Plugins/Experimental/Text3D/Source/Text3D/Private/Text3DComponent.cpp
--- a/Engine/Plugins/Experimental/Text3D/Source/Text3D/Private/Text3DComponent.cpp
+++ b/Engine/Plugins/Experimental/Text3D/Source/Text3D/Private/Text3DComponent.cpp
@@ -290,7 +290,7 @@ const FText& UText3DComponent::GetText() const
 }
 
 void UText3DComponent::SetText(const FText& Value)
-{
+{	
 	if (!Text.EqualTo(Value))
 	{
 		Text = Value;
@@ -1011,6 +1011,21 @@ void UText3DComponent::ClearTextMesh()
 	}
 }
 
+void UText3DComponent::TriggerInternalRebuild(const EText3DModifyFlags InModifyFlags)
+{
+	if (EnumHasAnyFlags(InModifyFlags, EText3DModifyFlags::Geometry))
+	{
+		MarkForGeometryUpdate();
+	}
+
+	if (EnumHasAnyFlags(ModifyFlags, EText3DModifyFlags::Layout))
+	{
+		MarkForLayoutUpdate();
+	}
+
+	RebuildInternal();
+}
+
 void UText3DComponent::BuildTextMesh(const bool& bCleanCache)
 {
 	// If we're already building, or have a build pending, don't do anything.
@@ -1089,7 +1104,9 @@ void UText3DComponent::BuildTextMeshInternal(const bool& bCleanCache)
 		TextLayoutMarshaller = FPlainTextLayoutMarshaller::Create();
 	}
 
-	FTextShaper::Get()->ShapeBidirectionalText(Style, Text.ToString(), TextLayout, TextLayoutMarshaller, ShapedText->Lines);
+	FText FormattedText = Text;
+	FormatText(FormattedText);
+	FTextShaper::Get()->ShapeBidirectionalText(Style, FormattedText.ToString(), TextLayout, TextLayoutMarshaller, ShapedText->Lines);
 
 	const int32 ApproximateGlyphNum = Algo::TransformAccumulate(ShapedText->Lines,
 		[](const FShapedGlyphLine& InLine)
@@ -1249,6 +1266,14 @@ void UText3DComponent::UpdateMaterial(const EText3DGroupType Type, UMaterialInte
 	}
 }
 
+FText UText3DComponent::GetFormattedText() const
+{
+	FText FormattedText = Text;
+	FormatText(FormattedText);
+
+	return FormattedText;
+}
+
 void UText3DComponent::OnVisibilityChanged()
 {
 	Super::OnVisibilityChanged();
diff --git a/Engine/Plugins/Experimental/Text3D/Source/Text3D/Public/Text3DComponent.h b/Engine/Plugins/Experimental/Text3D/Source/Text3D/Public/Text3DComponent.h
--- a/Engine/Plugins/Experimental/Text3D/Source/Text3D/Public/Text3DComponent.h
+++ b/Engine/Plugins/Experimental/Text3D/Source/Text3D/Public/Text3DComponent.h
@@ -376,6 +376,13 @@ protected:
 	UPROPERTY(EditAnywhere, BlueprintReadWrite, Getter = "CastsShadow", Setter = "SetCastShadow", Category = "Lighting", meta = (AllowPrivateAccess))
 	bool bCastShadow = true;
 
+	/**
+	 * Returns the Text property, after being formatted by the FormatText virtual function.
+	 * If FormatText is not overriden, the return FText will be the same as the Text property.
+	 */
+	UFUNCTION(BlueprintCallable, Category = "Text3D")
+	FText GetFormattedText() const;
+	
 protected:
 	/** Intercept and propagate a change on this component to all children. */
 	virtual void OnVisibilityChanged() override;
@@ -383,9 +390,19 @@ protected:
 	/** Intercept and propagate a change on this component to all children. */
 	virtual void OnHiddenInGameChanged() override;
 
+	/**
+	 * Will be called when text geometry is generated.
+	 * Override it to customize text formatting in the final geometry, without affecting the Text property.
+	 * Use GetFormattedText() to retrieve a FText with the result of this formatting.
+	 */
+	virtual void FormatText(FText& InOutText) const {}
+
 	/** Clears all generated components and meshes from this component. */
 	void ClearTextMesh();
 
+	/** Can be used to force an internal geometry and/or layout rebuild */
+	void TriggerInternalRebuild(const EText3DModifyFlags InModifyFlags);
+
 private:
 	UPROPERTY()
 	TObjectPtr<USceneComponent> TextRoot;
