[Text3D] Text3DComponent
Changed category based on feedback

diff --git a/Engine/Plugins/Experimental/Text3D/Source/Text3D/Public/Text3DComponent.h b/Engine/Plugins/Experimental/Text3D/Source/Text3D/Public/Text3DComponent.h
--- a/Engine/Plugins/Experimental/Text3D/Source/Text3D/Public/Text3DComponent.h
+++ b/Engine/Plugins/Experimental/Text3D/Source/Text3D/Public/Text3DComponent.h
@@ -44,7 +44,7 @@ enum class EText3DModifyFlags : uint8
 };
 ENUM_CLASS_FLAGS(EText3DModifyFlags)
 
-UCLASS(ClassGroup = (Text3D), meta = (BlueprintSpawnableComponent))
+UCLASS(ClassGroup = (Text3D), PrioritizeCategories = "Text Layout Geometry Materials", meta = (BlueprintSpawnableComponent))
 class TEXT3D_API UText3DComponent : public USceneComponent
 {
 	GENERATED_BODY()
@@ -283,91 +283,91 @@ protected:
 	bool bRefreshOnChange = true;
 
 	/** The text to generate a 3d mesh */
-	UPROPERTY(EditAnywhere, BlueprintReadWrite, Getter, Setter, Category = "Text3D", meta = (MultiLine = true, AllowPrivateAccess))
+	UPROPERTY(EditAnywhere, BlueprintReadWrite, Getter, Setter, Category = "Text", meta = (MultiLine = true, AllowPrivateAccess))
 	FText Text;
 
 	/** Size of the extrude */
-	UPROPERTY(EditAnywhere, BlueprintReadWrite, Getter, Setter, Category = "Text3D", meta = (ClampMin = 0, AllowPrivateAccess))
+	UPROPERTY(EditAnywhere, BlueprintReadWrite, Getter, Setter, Category = "Geometry", meta = (ClampMin = 0, AllowPrivateAccess))
 	float Extrude;
 
 	/** Size of bevel */
-	UPROPERTY(EditAnywhere, BlueprintReadWrite, Getter, Setter, Category = "Text3D", meta = (EditCondition = "!bOutline", ClampMin = 0, AllowPrivateAccess))
+	UPROPERTY(EditAnywhere, BlueprintReadWrite, Getter, Setter, Category = "Geometry", meta = (EditCondition = "!bOutline", ClampMin = 0, AllowPrivateAccess))
 	float Bevel;
 
 	/** Bevel Type */
-	UPROPERTY(EditAnywhere, BlueprintReadWrite, Getter, Setter, Category = "Text3D", meta = (EditCondition = "!bOutline", AllowPrivateAccess))
+	UPROPERTY(EditAnywhere, BlueprintReadWrite, Getter, Setter, Category = "Geometry", meta = (EditCondition = "!bOutline", AllowPrivateAccess))
 	EText3DBevelType BevelType;
 
 	/** Bevel Segments (Defines the amount of tesselation for the bevel part) */
-	UPROPERTY(EditAnywhere, BlueprintReadWrite, Getter, Setter, Category = "Text3D", meta = (EditCondition = "!bOutline", ClampMin = 1, ClampMax = 15, AllowPrivateAccess))
+	UPROPERTY(EditAnywhere, BlueprintReadWrite, Getter, Setter, Category = "Geometry", meta = (EditCondition = "!bOutline", ClampMin = 1, ClampMax = 15, AllowPrivateAccess))
 	int32 BevelSegments;
 
 	/** Generate Outline */
-	UPROPERTY(EditAnywhere, BlueprintReadWrite, Getter = "HasOutline", Setter = "SetHasOutline", Category = "Text3D", meta = (AllowPrivateAccess))
+	UPROPERTY(EditAnywhere, BlueprintReadWrite, Getter = "HasOutline", Setter = "SetHasOutline", Category = "Geometry", meta = (AllowPrivateAccess))
 	bool bOutline;
 
 	/** Outline expand/offset amount */
-	UPROPERTY(EditAnywhere, BlueprintReadWrite, Getter, Setter, Category = "Text3D", meta = (EditCondition = "bOutline", AllowPrivateAccess))
+	UPROPERTY(EditAnywhere, BlueprintReadWrite, Getter, Setter, Category = "Geometry", meta = (EditCondition = "bOutline", AllowPrivateAccess))
 	float OutlineExpand;
 
 	/** Material for the front part */
-	UPROPERTY(EditAnywhere, BlueprintReadWrite, Getter, Setter, Category = "Text3D", meta = (AllowPrivateAccess))
+	UPROPERTY(EditAnywhere, BlueprintReadWrite, Getter, Setter, Category = "Materials", meta = (AllowPrivateAccess))
 	TObjectPtr<UMaterialInterface> FrontMaterial;
 
 	/** Material for the bevel part */
-	UPROPERTY(EditAnywhere, BlueprintReadWrite, Getter, Setter, Category = "Text3D", meta = (EditCondition = "!bOutline", AllowPrivateAccess))
+	UPROPERTY(EditAnywhere, BlueprintReadWrite, Getter, Setter, Category = "Materials", meta = (EditCondition = "!bOutline", AllowPrivateAccess))
 	TObjectPtr<UMaterialInterface> BevelMaterial;
 
 	/** Material for the extruded part */
-	UPROPERTY(EditAnywhere, BlueprintReadWrite, Getter, Setter, Category = "Text3D", meta = (AllowPrivateAccess))
+	UPROPERTY(EditAnywhere, BlueprintReadWrite, Getter, Setter, Category = "Materials", meta = (AllowPrivateAccess))
 	TObjectPtr<UMaterialInterface> ExtrudeMaterial;
 
 	/** Material for the back part */
-	UPROPERTY(EditAnywhere, BlueprintReadWrite, Getter, Setter, Category = "Text3D", meta = (AllowPrivateAccess))
+	UPROPERTY(EditAnywhere, BlueprintReadWrite, Getter, Setter, Category = "Materials", meta = (AllowPrivateAccess))
 	TObjectPtr<UMaterialInterface> BackMaterial;
 
 	/** Text font */
-	UPROPERTY(EditAnywhere, BlueprintReadWrite, Getter, Setter, Category = "Text3D", meta = (AllowPrivateAccess))
+	UPROPERTY(EditAnywhere, BlueprintReadWrite, Getter, Setter, Category = "Text", meta = (AllowPrivateAccess))
 	TObjectPtr<UFont> Font;
 
 	/** Horizontal text alignment */
-	UPROPERTY(EditAnywhere, BlueprintReadWrite, Getter, Setter, Category = "Text3D", meta = (AllowPrivateAccess))
+	UPROPERTY(EditAnywhere, BlueprintReadWrite, Getter, Setter, Category = "Layout", meta = (AllowPrivateAccess))
 	EText3DHorizontalTextAlignment HorizontalAlignment;
 
 	/** Vertical text alignment */
-	UPROPERTY(EditAnywhere, BlueprintReadWrite, Getter, Setter, Category = "Text3D", meta = (AllowPrivateAccess))
+	UPROPERTY(EditAnywhere, BlueprintReadWrite, Getter, Setter, Category = "Layout", meta = (AllowPrivateAccess))
 	EText3DVerticalTextAlignment VerticalAlignment;
 
 	/** Text kerning */
-	UPROPERTY(EditAnywhere, BlueprintReadWrite, Getter, Setter, Category = "Text3D", meta = (AllowPrivateAccess))
+	UPROPERTY(EditAnywhere, BlueprintReadWrite, Getter, Setter, Category = "Layout", meta = (AllowPrivateAccess))
 	float Kerning;
 
 	/** Extra line spacing */
-	UPROPERTY(EditAnywhere, BlueprintReadWrite, Getter, Setter, Category = "Text3D", meta = (AllowPrivateAccess))
+	UPROPERTY(EditAnywhere, BlueprintReadWrite, Getter, Setter, Category = "Layout", meta = (AllowPrivateAccess))
 	float LineSpacing;
 
 	/** Extra word spacing */
-	UPROPERTY(EditAnywhere, BlueprintReadWrite, Setter, Category = "Text3D", meta = (AllowPrivateAccess))
+	UPROPERTY(EditAnywhere, BlueprintReadWrite, Setter, Category = "Layout", meta = (AllowPrivateAccess))
 	float WordSpacing;
 
 	/** Enables a maximum width to the 3D Text */
-	UPROPERTY(EditAnywhere, BlueprintReadWrite, Getter = "HasMaxWidth", Setter = "SetHasMaxWidth", Category = "Text3D", meta = (InlineEditConditionToggle, AllowPrivateAccess))
+	UPROPERTY(EditAnywhere, BlueprintReadWrite, Getter = "HasMaxWidth", Setter = "SetHasMaxWidth", Category = "Layout", meta = (InlineEditConditionToggle, AllowPrivateAccess))
 	bool bHasMaxWidth;
 
 	/** Sets a maximum width to the 3D Text */
-	UPROPERTY(EditAnywhere, BlueprintReadWrite, Getter, Setter, Category = "Text3D", meta = (EditCondition = "bHasMaxWidth", ClampMin = 1, AllowPrivateAccess))
+	UPROPERTY(EditAnywhere, BlueprintReadWrite, Getter, Setter, Category = "Layout", meta = (EditCondition = "bHasMaxWidth", ClampMin = 1, AllowPrivateAccess))
 	float MaxWidth;
 
 	/** Enables a maximum height to the 3D Text */
-	UPROPERTY(EditAnywhere, BlueprintReadWrite, Getter = "HasMaxHeight", Setter = "SetHasMaxHeight", Category = "Text3D", meta = (InlineEditConditionToggle, AllowPrivateAccess))
+	UPROPERTY(EditAnywhere, BlueprintReadWrite, Getter = "HasMaxHeight", Setter = "SetHasMaxHeight", Category = "Layout", meta = (InlineEditConditionToggle, AllowPrivateAccess))
 	bool bHasMaxHeight;
 
 	/** Sets a maximum height to the 3D Text */
-	UPROPERTY(EditAnywhere, BlueprintReadWrite, Getter, Setter, Category = "Text3D", meta = (EditCondition = "bHasMaxHeight", ClampMin = 1, AllowPrivateAccess))
+	UPROPERTY(EditAnywhere, BlueprintReadWrite, Getter, Setter, Category = "Layout", meta = (EditCondition = "bHasMaxHeight", ClampMin = 1, AllowPrivateAccess))
 	float MaxHeight;
 
 	/** Should the mesh scale proportionally when Max Width/Height is set */
-	UPROPERTY(EditAnywhere, BlueprintReadWrite, Getter = "ScalesProportionally", Setter = "SetScaleProportionally", Category = "Text3D", meta = (AllowPrivateAccess))
+	UPROPERTY(EditAnywhere, BlueprintReadWrite, Getter = "ScalesProportionally", Setter = "SetScaleProportionally", Category = "Layout", meta = (AllowPrivateAccess))
 	bool bScaleProportionally;
 
 	// Lighting flags
