[Text3D] Text meshes attached to editor only Kerning Component

Kerning Components were marked as editor only by the SetIsVisualizationComponent(true) call.

Mesh components are attached to them, so that was causing the error to appear.

Since they both sets of components are flagged as RF_Transient, removing the SetIsVisualizationComponent allows to avoid the error messages.

This shouldn't pose an issue for existing text assets geometry, which is generated again anyway OnRegister.

diff --git a/Engine/Plugins/Experimental/Text3D/Source/Text3D/Private/Text3DComponent.cpp b/Engine/Plugins/Experimental/Text3D/Source/Text3D/Private/Text3DComponent.cpp
--- a/Engine/Plugins/Experimental/Text3D/Source/Text3D/Private/Text3DComponent.cpp
+++ b/Engine/Plugins/Experimental/Text3D/Source/Text3D/Private/Text3DComponent.cpp
@@ -473,9 +473,6 @@ bool UText3DComponent::AllocateGlyphs(int32 Num)
 			const FName CharacterKerningComponentName = MakeUniqueObjectName(this, USceneComponent::StaticClass(), FName(*FString::Printf(TEXT("CharacterKerning%d"), GlyphId)));
 			USceneComponent* CharacterKerningComponent = NewObject<USceneComponent>(this, CharacterKerningComponentName, RF_Transient);
 			
-#if WITH_EDITOR
-			CharacterKerningComponent->SetIsVisualizationComponent(true);
-#endif
 			CharacterKerningComponent->AttachToComponent(TextRoot, FAttachmentTransformRules::KeepRelativeTransform);
 			CharacterKerningComponent->RegisterComponent();
 			CharacterKernings.Add(CharacterKerningComponent);
