Text3D Keep glyph static meshes accessible on CPU.

This is fixing the following warnings:
LogPhysics: Warning: UBodySetup::GetCookInfo: ContainsPhysicsTriMeshData returned true, but GetPhysicsTriMeshData returned false. This inconsistency should be fixed for asset [...].
LogStaticMesh: Warning: UStaticMesh::GetPhysicsTriMeshData: CPU data not available on selected LOD (UseLODIndex=0, LODForCollision=0) on 'StaticMesh [...].

#rnx

diff --git a/Engine/Plugins/Experimental/Text3D/Source/Text3D/Private/Glyph.cpp b/Engine/Plugins/Experimental/Text3D/Source/Text3D/Private/Glyph.cpp
--- a/Engine/Plugins/Experimental/Text3D/Source/Text3D/Private/Glyph.cpp
+++ b/Engine/Plugins/Experimental/Text3D/Source/Text3D/Private/Glyph.cpp
@@ -32,6 +32,7 @@ void FText3DGlyph::Build(UStaticMesh* StaticMesh, UMaterial* DefaultMaterial)
 	UStaticMesh::FBuildMeshDescriptionsParams Params;
 	Params.bCommitMeshDescription = true;
 	Params.bFastBuild = true;
+	Params.bAllowCpuAccess = true;
 	StaticMesh->BuildFromMeshDescriptions(MeshDescriptions, Params);
 }
 
