Text3DComponent - Fix for WebRemoteControl to be able to access the properties.

The "AllowPrivateAccess" meta was missing "= true", which is checked by RC in ResolveObjectProperty. In IsPropertyAllowed, GetBoolMetaData needs to return "true", and empty string is considered false.

#rnx

diff --git a/Engine/Plugins/Experimental/Text3D/Source/Text3D/Public/Text3DComponent.h b/Engine/Plugins/Experimental/Text3D/Source/Text3D/Public/Text3DComponent.h
--- a/Engine/Plugins/Experimental/Text3D/Source/Text3D/Public/Text3DComponent.h
+++ b/Engine/Plugins/Experimental/Text3D/Source/Text3D/Public/Text3DComponent.h
@@ -279,101 +279,101 @@ public:
 
 protected:
 	/** Whether to allow automatic refresh/mesh generation */
-	UPROPERTY(EditAnywhere, BlueprintReadWrite, Getter = "RefreshesOnChange", Setter = "SetRefreshOnChange", Category = "Text3D", AdvancedDisplay, meta = (AllowPrivateAccess))
+	UPROPERTY(EditAnywhere, BlueprintReadWrite, Getter = "RefreshesOnChange", Setter = "SetRefreshOnChange", Category = "Text3D", AdvancedDisplay, meta = (AllowPrivateAccess = "true"))
 	bool bRefreshOnChange = true;
 
 	/** The text to generate a 3d mesh */
-	UPROPERTY(EditAnywhere, BlueprintReadWrite, Getter, Setter, Category = "Text", meta = (MultiLine = true, AllowPrivateAccess))
+	UPROPERTY(EditAnywhere, BlueprintReadWrite, Getter, Setter, Category = "Text", meta = (MultiLine = true, AllowPrivateAccess = "true"))
 	FText Text;
 
 	/** Size of the extrude */
-	UPROPERTY(EditAnywhere, BlueprintReadWrite, Getter, Setter, Category = "Geometry", meta = (ClampMin = 0, AllowPrivateAccess))
+	UPROPERTY(EditAnywhere, BlueprintReadWrite, Getter, Setter, Category = "Geometry", meta = (ClampMin = 0, AllowPrivateAccess = "true"))
 	float Extrude;
 
 	/** Size of bevel */
-	UPROPERTY(EditAnywhere, BlueprintReadWrite, Getter, Setter, Category = "Geometry", meta = (EditCondition = "!bOutline", ClampMin = 0, AllowPrivateAccess))
+	UPROPERTY(EditAnywhere, BlueprintReadWrite, Getter, Setter, Category = "Geometry", meta = (EditCondition = "!bOutline", ClampMin = 0, AllowPrivateAccess = "true"))
 	float Bevel;
 
 	/** Bevel Type */
-	UPROPERTY(EditAnywhere, BlueprintReadWrite, Getter, Setter, Category = "Geometry", meta = (EditCondition = "!bOutline", AllowPrivateAccess))
+	UPROPERTY(EditAnywhere, BlueprintReadWrite, Getter, Setter, Category = "Geometry", meta = (EditCondition = "!bOutline", AllowPrivateAccess = "true"))
 	EText3DBevelType BevelType;
 
 	/** Bevel Segments (Defines the amount of tesselation for the bevel part) */
-	UPROPERTY(EditAnywhere, BlueprintReadWrite, Getter, Setter, Category = "Geometry", meta = (EditCondition = "!bOutline", ClampMin = 1, ClampMax = 15, AllowPrivateAccess))
+	UPROPERTY(EditAnywhere, BlueprintReadWrite, Getter, Setter, Category = "Geometry", meta = (EditCondition = "!bOutline", ClampMin = 1, ClampMax = 15, AllowPrivateAccess = "true"))
 	int32 BevelSegments;
 
 	/** Generate Outline */
-	UPROPERTY(EditAnywhere, BlueprintReadWrite, Getter = "HasOutline", Setter = "SetHasOutline", Category = "Geometry", meta = (AllowPrivateAccess))
+	UPROPERTY(EditAnywhere, BlueprintReadWrite, Getter = "HasOutline", Setter = "SetHasOutline", Category = "Geometry", meta = (AllowPrivateAccess = "true"))
 	bool bOutline;
 
 	/** Outline expand/offset amount */
-	UPROPERTY(EditAnywhere, BlueprintReadWrite, Getter, Setter, Category = "Geometry", meta = (EditCondition = "bOutline", AllowPrivateAccess))
+	UPROPERTY(EditAnywhere, BlueprintReadWrite, Getter, Setter, Category = "Geometry", meta = (EditCondition = "bOutline", AllowPrivateAccess = "true"))
 	float OutlineExpand;
 
 	/** Material for the front part */
-	UPROPERTY(EditAnywhere, BlueprintReadWrite, Getter, Setter, Category = "Materials", meta = (AllowPrivateAccess))
+	UPROPERTY(EditAnywhere, BlueprintReadWrite, Getter, Setter, Category = "Materials", meta = (AllowPrivateAccess = "true"))
 	TObjectPtr<UMaterialInterface> FrontMaterial;
 
 	/** Material for the bevel part */
-	UPROPERTY(EditAnywhere, BlueprintReadWrite, Getter, Setter, Category = "Materials", meta = (EditCondition = "!bOutline", AllowPrivateAccess))
+	UPROPERTY(EditAnywhere, BlueprintReadWrite, Getter, Setter, Category = "Materials", meta = (EditCondition = "!bOutline", AllowPrivateAccess = "true"))
 	TObjectPtr<UMaterialInterface> BevelMaterial;
 
 	/** Material for the extruded part */
-	UPROPERTY(EditAnywhere, BlueprintReadWrite, Getter, Setter, Category = "Materials", meta = (AllowPrivateAccess))
+	UPROPERTY(EditAnywhere, BlueprintReadWrite, Getter, Setter, Category = "Materials", meta = (AllowPrivateAccess = "true"))
 	TObjectPtr<UMaterialInterface> ExtrudeMaterial;
 
 	/** Material for the back part */
-	UPROPERTY(EditAnywhere, BlueprintReadWrite, Getter, Setter, Category = "Materials", meta = (AllowPrivateAccess))
+	UPROPERTY(EditAnywhere, BlueprintReadWrite, Getter, Setter, Category = "Materials", meta = (AllowPrivateAccess = "true"))
 	TObjectPtr<UMaterialInterface> BackMaterial;
 
 	/** Text font */
-	UPROPERTY(EditAnywhere, BlueprintReadWrite, Getter, Setter, Category = "Text", meta = (AllowPrivateAccess))
+	UPROPERTY(EditAnywhere, BlueprintReadWrite, Getter, Setter, Category = "Text", meta = (AllowPrivateAccess = "true"))
 	TObjectPtr<UFont> Font;
 
 	/** Horizontal text alignment */
-	UPROPERTY(EditAnywhere, BlueprintReadWrite, Getter, Setter, Category = "Layout", meta = (AllowPrivateAccess))
+	UPROPERTY(EditAnywhere, BlueprintReadWrite, Getter, Setter, Category = "Layout", meta = (AllowPrivateAccess = "true"))
 	EText3DHorizontalTextAlignment HorizontalAlignment;
 
 	/** Vertical text alignment */
-	UPROPERTY(EditAnywhere, BlueprintReadWrite, Getter, Setter, Category = "Layout", meta = (AllowPrivateAccess))
+	UPROPERTY(EditAnywhere, BlueprintReadWrite, Getter, Setter, Category = "Layout", meta = (AllowPrivateAccess = "true"))
 	EText3DVerticalTextAlignment VerticalAlignment;
 
 	/** Text kerning */
-	UPROPERTY(EditAnywhere, BlueprintReadWrite, Getter, Setter, Category = "Layout", meta = (AllowPrivateAccess))
+	UPROPERTY(EditAnywhere, BlueprintReadWrite, Getter, Setter, Category = "Layout", meta = (AllowPrivateAccess = "true"))
 	float Kerning;
 
 	/** Extra line spacing */
-	UPROPERTY(EditAnywhere, BlueprintReadWrite, Getter, Setter, Category = "Layout", meta = (AllowPrivateAccess))
+	UPROPERTY(EditAnywhere, BlueprintReadWrite, Getter, Setter, Category = "Layout", meta = (AllowPrivateAccess = "true"))
 	float LineSpacing;
 
 	/** Extra word spacing */
-	UPROPERTY(EditAnywhere, BlueprintReadWrite, Setter, Category = "Layout", meta = (AllowPrivateAccess))
+	UPROPERTY(EditAnywhere, BlueprintReadWrite, Setter, Category = "Layout", meta = (AllowPrivateAccess = "true"))
 	float WordSpacing;
 
 	/** Enables a maximum width to the 3D Text */
-	UPROPERTY(EditAnywhere, BlueprintReadWrite, Getter = "HasMaxWidth", Setter = "SetHasMaxWidth", Category = "Layout", meta = (InlineEditConditionToggle, AllowPrivateAccess))
+	UPROPERTY(EditAnywhere, BlueprintReadWrite, Getter = "HasMaxWidth", Setter = "SetHasMaxWidth", Category = "Layout", meta = (InlineEditConditionToggle, AllowPrivateAccess = "true"))
 	bool bHasMaxWidth;
 
 	/** Sets a maximum width to the 3D Text */
-	UPROPERTY(EditAnywhere, BlueprintReadWrite, Getter, Setter, Category = "Layout", meta = (EditCondition = "bHasMaxWidth", ClampMin = 1, AllowPrivateAccess))
+	UPROPERTY(EditAnywhere, BlueprintReadWrite, Getter, Setter, Category = "Layout", meta = (EditCondition = "bHasMaxWidth", ClampMin = 1, AllowPrivateAccess = "true"))
 	float MaxWidth;
 
 	/** Enables a maximum height to the 3D Text */
-	UPROPERTY(EditAnywhere, BlueprintReadWrite, Getter = "HasMaxHeight", Setter = "SetHasMaxHeight", Category = "Layout", meta = (InlineEditConditionToggle, AllowPrivateAccess))
+	UPROPERTY(EditAnywhere, BlueprintReadWrite, Getter = "HasMaxHeight", Setter = "SetHasMaxHeight", Category = "Layout", meta = (InlineEditConditionToggle, AllowPrivateAccess = "true"))
 	bool bHasMaxHeight;
 
 	/** Sets a maximum height to the 3D Text */
-	UPROPERTY(EditAnywhere, BlueprintReadWrite, Getter, Setter, Category = "Layout", meta = (EditCondition = "bHasMaxHeight", ClampMin = 1, AllowPrivateAccess))
+	UPROPERTY(EditAnywhere, BlueprintReadWrite, Getter, Setter, Category = "Layout", meta = (EditCondition = "bHasMaxHeight", ClampMin = 1, AllowPrivateAccess = "true"))
 	float MaxHeight;
 
 	/** Should the mesh scale proportionally when Max Width/Height is set */
-	UPROPERTY(EditAnywhere, BlueprintReadWrite, Getter = "ScalesProportionally", Setter = "SetScaleProportionally", Category = "Layout", meta = (AllowPrivateAccess))
+	UPROPERTY(EditAnywhere, BlueprintReadWrite, Getter = "ScalesProportionally", Setter = "SetScaleProportionally", Category = "Layout", meta = (AllowPrivateAccess = "true"))
 	bool bScaleProportionally;
 
 	// Lighting flags
 	
 	/** Controls whether the text glyphs should cast a shadow or not. */
-	UPROPERTY(EditAnywhere, BlueprintReadWrite, Getter = "CastsShadow", Setter = "SetCastShadow", Category = "Lighting", meta = (AllowPrivateAccess))
+	UPROPERTY(EditAnywhere, BlueprintReadWrite, Getter = "CastsShadow", Setter = "SetCastShadow", Category = "Lighting", meta = (AllowPrivateAccess = "true"))
 	bool bCastShadow = true;
 
 	/**
@@ -407,7 +407,7 @@ private:
 	UPROPERTY()
 	TObjectPtr<USceneComponent> TextRoot;
 
-	UPROPERTY(BlueprintAssignable, Category = Events, meta = (AllowPrivateAccess = true, DisplayName = "On Text Generated"))
+	UPROPERTY(BlueprintAssignable, Category = Events, meta = (AllowPrivateAccess = "true", DisplayName = "On Text Generated"))
 	FTextGenerated TextGeneratedDelegate;
 
 	FTextGeneratedNative TextGeneratedNativeDelegate;
