Fix C6011
#rnx

diff --git a/Engine/Plugins/Experimental/Text3D/Source/Text3D/Private/Text3DEngineSubsystem.cpp b/Engine/Plugins/Experimental/Text3D/Source/Text3D/Private/Text3DEngineSubsystem.cpp
--- a/Engine/Plugins/Experimental/Text3D/Source/Text3D/Private/Text3DEngineSubsystem.cpp
+++ b/Engine/Plugins/Experimental/Text3D/Source/Text3D/Private/Text3DEngineSubsystem.cpp
@@ -255,7 +255,7 @@ TSharedContourNode FCachedFontData::GetGlyphContours(uint32 GlyphIndex, const FF
 
 #if UE_BUILD_DEBUG
 	// If true, the symbol itself was probably bitmap-only, and this will have placeholder contours
-	bool bIsLastResort = FontFaceData->GetFace()->family_name
+	bool bIsLastResort = FontFaceData && FontFaceData->GetFace()->family_name
 		? FString(FontFaceData->GetFace()->family_name) == TEXT("LastResort")
 		: false;
 	if (bIsLastResort)
