[Text3D] Add Font typefaces support

- adds the possibility to choose between available font faces for the current font
- text3d subsystem font cache is updated so that it caches each typeface separately

diff --git a/Engine/Plugins/Experimental/Text3D/Source/Text3D/Private/Text3DComponent.cpp b/Engine/Plugins/Experimental/Text3D/Source/Text3D/Private/Text3DComponent.cpp
--- a/Engine/Plugins/Experimental/Text3D/Source/Text3D/Private/Text3DComponent.cpp
+++ b/Engine/Plugins/Experimental/Text3D/Source/Text3D/Private/Text3DComponent.cpp
@@ -3,12 +3,12 @@
 #include "Text3DComponent.h"
 
 #include "Algo/Count.h"
-#include "Algo/IndexOf.h"
 #include "Async/Async.h"
 #include "Components/StaticMeshComponent.h"
 #include "Engine/Engine.h"
 #include "Engine/Font.h"
 #include "Engine/StaticMesh.h"
+#include "Fonts/CompositeFont.h"
 #include "Fonts/SlateTextShaper.h"
 #include "Framework/Text/PlainTextLayoutMarshaller.h"
 #include "Glyph.h"
@@ -124,13 +124,15 @@ UText3DComponent::UText3DComponent()
 	ModifyFlags = EText3DModifyFlags::All;
 
 	TextScale = FVector::ZeroVector;
+
+	RefreshTypeface();
 }
 
 void UText3DComponent::PostLoad()
 {
 	// Reset so it's rebuilt (needed if re-using the component!)
 	ModifyFlags = EText3DModifyFlags::All;
-	
+
 	Super::PostLoad();
 }
 
@@ -166,6 +168,75 @@ void UText3DComponent::ClearUpdateFlags()
 	ModifyFlags = EText3DModifyFlags::None;
 }
 
+uint32 UText3DComponent::GetTypeFaceIndex() const
+{
+	uint32 TypefaceIndex = 0;
+
+	if (Font)
+	{
+		TArray<FTypefaceEntry>& Fonts = Font->CompositeFont.DefaultTypeface.Fonts;
+
+		for (uint32 Index = 0; Index < static_cast<uint32>(Fonts.Num()); Index++)
+		{
+			if (Typeface == Fonts[Index].Name)
+			{
+				TypefaceIndex = Index;
+				break;
+			}
+		}
+	}
+
+	return TypefaceIndex;
+}
+
+bool UText3DComponent::IsTypefaceAvailable(FName InTypeface) const
+{
+	for (const FTypefaceEntry& TypefaceElement : GetAvailableTypefaces())
+	{
+		if (InTypeface == TypefaceElement.Name)
+		{
+			return true;
+		}
+	}
+
+	return false;
+}
+
+TArray<FTypefaceEntry> UText3DComponent::GetAvailableTypefaces() const
+{
+	if (Font)
+	{
+		return Font->CompositeFont.DefaultTypeface.Fonts;
+	}
+
+	return {};
+}
+
+void UText3DComponent::RefreshTypeface()
+{
+	if (Font)
+	{
+		TArray<FTypefaceEntry>& Fonts = Font->CompositeFont.DefaultTypeface.Fonts;
+		for (int32 Index = 0; Index < Fonts.Num(); Index++)
+		{
+			if (Typeface == Fonts[Index].Name)
+			{
+				// Typeface stays the same
+				return;
+			}
+		}
+
+		if (!Fonts.IsEmpty())
+		{
+			Typeface = Fonts[0].Name;
+		}
+		else
+		{
+			Typeface = TEXT("");
+		}
+	}
+}
+
 void UText3DComponent::OnRegister()
 {
 	Super::OnRegister();
@@ -229,7 +300,12 @@ void UText3DComponent::PostEditChangeProperty(FPropertyChangedEvent& PropertyCha
 		// Force minimum bevel segments based on the BevelType
 		SetBevelSegments(BevelSegments);
 	}
-	else if (Name == GET_MEMBER_NAME_CHECKED(UText3DComponent, Font) ||
+	else if (Name == GET_MEMBER_NAME_CHECKED(UText3DComponent, Font))
+	{
+		MarkForGeometryUpdate();
+		RefreshTypeface();
+	}
+	else if (Name == GET_MEMBER_NAME_CHECKED(UText3DComponent, Typeface) ||
 			 Name == GET_MEMBER_NAME_CHECKED(UText3DComponent, Text) ||
 			 Name == GET_MEMBER_NAME_CHECKED(UText3DComponent, OutlineExpand) ||
 			 Name == GET_MEMBER_NAME_CHECKED(UText3DComponent, bOutline) ||
@@ -290,7 +366,7 @@ const FText& UText3DComponent::GetText() const
 }
 
 void UText3DComponent::SetText(const FText& Value)
-{	
+{
 	if (!Text.EqualTo(Value))
 	{
 		Text = Value;
@@ -309,6 +385,7 @@ void UText3DComponent::SetFont(UFont* InFont)
 	if (Font != InFont)
 	{
 		Font = InFont;
+		RefreshTypeface();
 		MarkForGeometryUpdate();
 		RebuildInternal();
 	}
@@ -324,7 +401,7 @@ void UText3DComponent::SetHasOutline(const bool bValue)
 	if (bOutline != bValue)
 	{
 		bOutline = bValue;
-		MarkForGeometryUpdate();	
+		MarkForGeometryUpdate();
 		RebuildInternal();
 	}
 }
@@ -472,7 +549,7 @@ bool UText3DComponent::AllocateGlyphs(int32 Num)
 
 			const FName CharacterKerningComponentName = MakeUniqueObjectName(this, USceneComponent::StaticClass(), FName(*FString::Printf(TEXT("CharacterKerning%d"), GlyphId)));
 			USceneComponent* CharacterKerningComponent = NewObject<USceneComponent>(this, CharacterKerningComponentName, RF_Transient);
-			
+
 			CharacterKerningComponent->AttachToComponent(TextRoot, FAttachmentTransformRules::KeepRelativeTransform);
 			CharacterKerningComponent->RegisterComponent();
 			CharacterKernings.Add(CharacterKerningComponent);
@@ -503,7 +580,7 @@ bool UText3DComponent::AllocateGlyphs(int32 Num)
 				CharacterKerningComponent->UnregisterComponent();
 				CharacterKerningComponent->DestroyComponent();
 			}
-			
+
 			UStaticMeshComponent* StaticMeshComponent = CharacterMeshes[CharacterIndex];
 			if (IsValid(StaticMeshComponent))
 			{
@@ -512,7 +589,7 @@ bool UText3DComponent::AllocateGlyphs(int32 Num)
 				StaticMeshComponent->DestroyComponent();
 			}
 		}
-		
+
 		CharacterKernings.RemoveAt(CharacterKernings.Num() - 1 - DeltaNum, DeltaNum);
 		CharacterMeshes.RemoveAt(CharacterMeshes.Num() - 1 - DeltaNum, DeltaNum);
 	}
@@ -543,7 +620,7 @@ UMaterialInterface* UText3DComponent::GetMaterial(const EText3DGroupType Type) c
 	{
 		return BackMaterial;
 	}
-		
+
 	default:
 	{
 		return nullptr;
@@ -767,7 +844,7 @@ void UText3DComponent::SetCastShadow(bool NewCastShadow)
 		{
 			MeshComponent->SetCastShadow(bCastShadow);
 		}
-		
+
 		MarkRenderStateDirty();
 	}
 }
@@ -807,6 +884,19 @@ const TArray<UStaticMeshComponent*>& UText3DComponent::GetGlyphMeshComponents()
 	return CharacterMeshes;
 }
 
+void UText3DComponent::SetTypeface(const FName InTypeface)
+{
+	if (Typeface == InTypeface || !IsTypefaceAvailable(InTypeface))
+	{
+		return;
+	}
+
+	Typeface = InTypeface;
+
+	MarkForGeometryUpdate();
+	RebuildInternal();
+}
+
 void UText3DComponent::Rebuild()
 {
 	// Rebuild, ignoring RefreshOnChange
@@ -1030,7 +1120,7 @@ void UText3DComponent::BuildTextMesh(const bool& bCleanCache)
 	{
 		return;
 	}
-	
+
 	bIsBuilding = true;
 
 	TWeakObjectPtr<UText3DComponent> WeakThis(this);
@@ -1066,18 +1156,21 @@ void UText3DComponent::BuildTextMeshInternal(const bool& bCleanCache)
 	{
 		return;
 	}
-	
+
 	UText3DEngineSubsystem* Subsystem = GEngine->GetEngineSubsystem<UText3DEngineSubsystem>();
-	FCachedFontData& CachedFontData = Subsystem->GetCachedFontData(Font);
-	const FT_Face Face = CachedFontData.GetFreeTypeFace();
+
+	const uint32 TypefaceIndex = GetTypeFaceIndex();
+	FCachedFontData& CachedFontData = Subsystem->GetCachedFontData(Font, TypefaceIndex);
+	const FT_Face Face = CachedFontData.GetFreeTypeFace(TypefaceIndex);
 	if (!Face)
 	{ 
 		UE_LOG(LogText3D, Error, TEXT("Failed to load font data '%s'"), *CachedFontData.GetFontName());
 		return;
 	}
 
-	CachedCounterReferences.Add(CachedFontData.GetCacheCounter());
-	CachedCounterReferences.Add(CachedFontData.GetMeshesCacheCounter(FGlyphMeshParameters{Extrude, Bevel, BevelType, BevelSegments, bOutline, OutlineExpand}));
+	FGlyphMeshParameters GlyphMeshParameters{Extrude, Bevel, BevelType, BevelSegments, bOutline, OutlineExpand, TypefaceIndex};
+	CachedCounterReferences.Add(CachedFontData.GetCacheCounter(TypefaceIndex));
+	CachedCounterReferences.Add(CachedFontData.GetMeshesCacheCounter(GlyphMeshParameters));
 
 	ShapedText->Reset();
 	ShapedText->LineHeight = Face->size->metrics.height * FontInverseScale;
@@ -1087,7 +1180,8 @@ void UText3DComponent::BuildTextMeshInternal(const bool& bCleanCache)
 	constexpr int32 AdjustedFontSize = 48; // Magic number that makes font scale consistent with previous implementation
 	FSlateFontInfo FontInfo(Font, AdjustedFontSize);
 	FontInfo.CompositeFont = FStyleDefaults::GetFontInfo().CompositeFont;
-	
+	FontInfo.TypefaceFontName = Typeface;
+
 	FTextBlockStyle Style;
 	Style.SetFont(FontInfo);
 
@@ -1126,7 +1220,7 @@ void UText3DComponent::BuildTextMeshInternal(const bool& bCleanCache)
 				{
 					GlyphIndexToFontFace.Add(GlyphEntry.GlyphIndex, nullptr);
 				}
-				
+
 				continue;
 			}
 
@@ -1137,7 +1231,7 @@ void UText3DComponent::BuildTextMeshInternal(const bool& bCleanCache)
 			}
 		}
 	}
-	
+
 	CalculateTextWidth();
 	CalculateTextScale();
 	TextRoot->SetRelativeScale3D(GetTextScale());
@@ -1173,7 +1267,7 @@ void UText3DComponent::BuildTextMeshInternal(const bool& bCleanCache)
 			const int32 GlyphId = GlyphIndex++;
 
 			const FFreeTypeFace* FontFace = GlyphIndexToFontFace[ShapedGlyph.GlyphIndex];
-			UStaticMesh* CachedMesh = CachedFontData.GetGlyphMesh(ShapedGlyph.GlyphIndex, FGlyphMeshParameters{Extrude, Bevel, BevelType, BevelSegments, bOutline, OutlineExpand}, FontFace);
+			UStaticMesh* CachedMesh = CachedFontData.GetGlyphMesh(ShapedGlyph.GlyphIndex, GlyphMeshParameters, FontFace);
 			if (!CachedMesh || FMath::IsNearlyZero(CachedMesh->GetBounds().SphereRadius))
 			{
 				continue;
@@ -1182,7 +1276,7 @@ void UText3DComponent::BuildTextMeshInternal(const bool& bCleanCache)
 			if (CharacterMeshes.IsValidIndex(GlyphId))
 			{
 				UStaticMeshComponent* StaticMeshComponent = CharacterMeshes[GlyphId];
-				StaticMeshComponent->SetStaticMesh(CachedMesh);			
+				StaticMeshComponent->SetStaticMesh(CachedMesh);
 				StaticMeshComponent->SetVisibility(GetVisibleFlag());
 				StaticMeshComponent->SetHiddenInGame(bHiddenInGame);
 				StaticMeshComponent->SetCastShadow(bCastShadow);
@@ -1207,7 +1301,7 @@ void UText3DComponent::BuildTextMeshInternal(const bool& bCleanCache)
 			}
 		}
 	}
-	
+
 	for (int32 Index = 0; Index < static_cast<int32>(EText3DGroupType::TypeCount); Index++)
 	{
 		const EText3DGroupType Type = static_cast<EText3DGroupType>(Index);
@@ -1216,7 +1310,7 @@ void UText3DComponent::BuildTextMeshInternal(const bool& bCleanCache)
 
 	TextGeneratedNativeDelegate.Broadcast();
 	TextGeneratedDelegate.Broadcast();
-	
+
 	ClearUpdateFlags();
 
 	if (bCleanCache)
@@ -1246,7 +1340,7 @@ void UText3DComponent::UpdateMaterial(const EText3DGroupType Type, UMaterialInte
 	{
 		return; 
 	}
-	
+
 	const bool bHasExtrude = !FMath::IsNearlyZero(Extrude);
 	if (!bHasExtrude && Type == EText3DGroupType::Extrude)
 	{
@@ -1271,6 +1365,21 @@ FText UText3DComponent::GetFormattedText() const
 	return FormattedText;
 }
 
+TArray<FName> UText3DComponent::GetTypefaceNames() const
+{
+	TArray<FName> TypefaceNames;
+
+	if (Font)
+	{
+		for (const FTypefaceEntry& TypeFaceFont : Font->CompositeFont.DefaultTypeface.Fonts)
+		{
+			TypefaceNames.Add(TypeFaceFont.Name);
+		}
+	}
+
+	return TypefaceNames;
+}
+
 void UText3DComponent::OnVisibilityChanged()
 {
 	Super::OnVisibilityChanged();
diff --git a/Engine/Plugins/Experimental/Text3D/Source/Text3D/Private/Text3DEngineSubsystem.cpp b/Engine/Plugins/Experimental/Text3D/Source/Text3D/Private/Text3DEngineSubsystem.cpp
--- a/Engine/Plugins/Experimental/Text3D/Source/Text3D/Private/Text3DEngineSubsystem.cpp
+++ b/Engine/Plugins/Experimental/Text3D/Source/Text3D/Private/Text3DEngineSubsystem.cpp
@@ -18,7 +18,7 @@ FCachedFontMeshes::FCachedFontMeshes()
 	CacheCounter = MakeShared<int32>(0);
 }
 
-int32 FCachedFontMeshes::GetCacheCount()
+int32 FCachedFontMeshes::GetCacheCount() const
 {
 	return CacheCounter.GetSharedReferenceCount();
 }
@@ -63,6 +63,35 @@ void UText3DEngineSubsystem::Deinitialize()
 	Super::Deinitialize();
 }
 
+bool UText3DEngineSubsystem::Exec(UWorld* InWorld, const TCHAR* InCmd, FOutputDevice& InAr)
+{
+	if (HasAllFlags(RF_ClassDefaultObject))
+	{
+		return false;
+	}
+
+	if (FParse::Command(&InCmd, TEXT("text3d")))
+	{
+		if (FParse::Command(&InCmd, TEXT("cache")))
+		{
+			if (FParse::Command(&InCmd, TEXT("show")))
+			{
+				PrintCache();
+				return true;
+			}
+			else if (FParse::Command(&InCmd, TEXT("clear")))
+			{
+				Cleanup();
+				PrintCache();
+
+				return true;
+			}
+		}
+	}
+
+	return false;
+}
+
 void UText3DEngineSubsystem::Reset()
 {
 	CachedFonts.Reset();
@@ -74,10 +103,20 @@ bool UText3DEngineSubsystem::CleanupTimerCallback(float DeltaTime)
 	return true;
 }
 
+void UText3DEngineSubsystem::PrintCache() const
+{
+	UE_LOG(LogText3D, Log, TEXT("Text3D Engine Subsystem is currently caching the following fonts:"));
+	for (const TPair<uint32, FCachedFontData>& CachedFont : CachedFonts)
+	{
+		CachedFont.Value.PrintCache();
+	}
+}
+
 void UText3DEngineSubsystem::Cleanup()
 {
 	QUICK_SCOPE_CYCLE_COUNTER(STAT_UText3DEngineSubsystem_Cleanup);
 
+	UE_LOG(LogText3D, Log, TEXT("Text3D Engine Subsystem is clearing unreferenced cached fonts and typefaces."));
 	for (auto It = CachedFonts.CreateIterator(); It; ++It)
 	{
 		if (It.Value().Cleanup())
@@ -87,43 +126,61 @@ void UText3DEngineSubsystem::Cleanup()
 	}
 }
 
-FCachedFontData& UText3DEngineSubsystem::GetCachedFontData(UFont* Font)
+FCachedFontData& UText3DEngineSubsystem::GetCachedFontData(UFont* InFont, int32 InTypefaceEntryIndex)
 {
-	uint32 FontHash = HashCombine(0, GetTypeHash(Font));
+	const uint32 FontHash = HashCombine(0, GetTypeHash(InFont));
+
 	if (CachedFonts.Contains(FontHash))
 	{
 		uint32 TypefaceFontDataHash = 0;
 
 		// First we check if the font itself has changed since we last cached it
-		const FCompositeFont* CompositeFont = Font->GetCompositeFont();
-		if (CompositeFont && CompositeFont->DefaultTypeface.Fonts.Num() > 0)
+		const FCompositeFont* CompositeFont = InFont->GetCompositeFont();
+		if (CompositeFont && !CompositeFont->DefaultTypeface.Fonts.IsEmpty())
 		{
-			const FTypefaceEntry& Typeface = CompositeFont->DefaultTypeface.Fonts[0];
+			FTypefaceEntry Typeface;
+			if (CompositeFont->DefaultTypeface.Fonts.IsValidIndex(InTypefaceEntryIndex))
+			{
+				Typeface = CompositeFont->DefaultTypeface.Fonts[InTypefaceEntryIndex];
+			}
+			else
+			{
+				Typeface = CompositeFont->DefaultTypeface.Fonts[0];
+			}
+
 			TypefaceFontDataHash = HashCombine(0, GetTypeHash(Typeface.Font));
 		}
 
-		if (CachedFonts[FontHash].GetTypefaceFontDataHash() != TypefaceFontDataHash)
+		const uint32 CachedTypefaceFontDataHash = CachedFonts[FontHash].GetTypefaceFontDataHash(InTypefaceEntryIndex);
+		if (CachedTypefaceFontDataHash != TypefaceFontDataHash)
 		{
-			CachedFonts.Remove(FontHash);
+			// FontFace has changed, remove the cached data
+			CachedFonts[FontHash].CleanupTypeface(InTypefaceEntryIndex);
 		}
 	}
 
 	if (!CachedFonts.Contains(FontHash))
 	{
-		FCachedFontData& CachedFontData = CachedFonts.Add(FontHash);
-		CachedFontData.Font = Font;
-		CachedFontData.LoadFreeTypeFace();
+		FCachedFontData& NewCachedFontData = CachedFonts.Add(FontHash);
+		NewCachedFontData.SetFont(InFont);
 	}
 
+	FCachedFontData& CachedFontData = CachedFonts[FontHash];
+	CachedFontData.LoadFreeTypeFace(InTypefaceEntryIndex);
+
 	return CachedFonts[FontHash];
 }
 
-FCachedFontData::FCachedFontData()
+FTypefaceFontData::FTypefaceFontData()
+	: FreeTypeFace(nullptr)
+	, TypefaceFontDataHash(0)
 {
-	TypefaceFontDataHash = 0;
 	CacheCounter = MakeShared<int32>(0);
+}
+
+FCachedFontData::FCachedFontData()
+{
 	Font = nullptr;
-	FreeTypeFace = nullptr;
 }
 
 FCachedFontData::~FCachedFontData()
@@ -133,74 +190,107 @@ FCachedFontData::~FCachedFontData()
 
 void FCachedFontData::ClearFreeTypeFace()
 {
-	TypefaceFontDataHash = 0;
-	if (FreeTypeFace)
-	{
-		FT_Done_Face(FreeTypeFace);
-		FreeTypeFace = nullptr;
-		Data.Reset();
-	}
+	TypefaceFontDataMap.Reset();
 }
 
-void FCachedFontData::LoadFreeTypeFace()
+void FCachedFontData::LoadFreeTypeFace(uint32 InTypefaceIndex)
 {
-	ClearFreeTypeFace();
+	const bool bContainsTypeface = TypefaceFontDataMap.Contains(InTypefaceIndex);
+	if (bContainsTypeface)
+	{
+		return;
+	}
+
 	if (!Font)
 	{
 		return;
 	}
 
 	const FCompositeFont* const CompositeFont = Font->GetCompositeFont();
-	if (!CompositeFont || CompositeFont->DefaultTypeface.Fonts.Num() == 0)
+	if (!CompositeFont || CompositeFont->DefaultTypeface.Fonts.IsEmpty())
+	{
+		return;
+	}
+
+	if (!CompositeFont->DefaultTypeface.Fonts.IsValidIndex(InTypefaceIndex))
 	{
 		return;
 	}
 
-	const FTypefaceEntry& Typeface = CompositeFont->DefaultTypeface.Fonts[0];
+	const FTypefaceEntry& Typeface = CompositeFont->DefaultTypeface.Fonts[InTypefaceIndex];
 	const FFontFaceDataConstPtr FaceData = Typeface.Font.GetFontFaceData();
 
+	if (!bContainsTypeface)
+	{
+		FTypefaceFontData& TypefaceFontData = TypefaceFontDataMap.Add(InTypefaceIndex);
+		TypefaceFontData.SetTypefaceName(Typeface.Name);
+	}
+
 	if (FaceData.IsValid() && FaceData->HasData() && FaceData->GetData().Num() > 0)
 	{
-		Data = FaceData->GetData();
-		FT_New_Memory_Face(FText3DModule::GetFreeTypeLibrary(), Data.GetData(), Data.Num(), 0, &FreeTypeFace);
+		TypefaceFontDataMap[InTypefaceIndex].SetTypefaceData(FaceData->GetData());
+		FT_New_Memory_Face(FText3DModule::GetFreeTypeLibrary(), TypefaceFontDataMap[InTypefaceIndex].GetTypefaceData().GetData(), TypefaceFontDataMap[InTypefaceIndex].GetTypefaceData().Num(), 0, &TypefaceFontDataMap[InTypefaceIndex].GetTypeface());
 	}
-	else if (FFileHelper::LoadFileToArray(Data, *Typeface.Font.GetFontFilename()) && Data.Num() > 0)
+	else if (FFileHelper::LoadFileToArray(TypefaceFontDataMap[InTypefaceIndex].GetTypefaceData(), *Typeface.Font.GetFontFilename()) && TypefaceFontDataMap[InTypefaceIndex].GetTypefaceData().Num() > 0)
 	{
-		FT_New_Memory_Face(FText3DModule::GetFreeTypeLibrary(), Data.GetData(), Data.Num(), 0, &FreeTypeFace);
+		FT_New_Memory_Face(FText3DModule::GetFreeTypeLibrary(), TypefaceFontDataMap[InTypefaceIndex].GetTypefaceData().GetData(), TypefaceFontDataMap[InTypefaceIndex].GetTypefaceData().Num(), 0, &TypefaceFontDataMap[InTypefaceIndex].GetTypeface());
 	}
 
-	if (FreeTypeFace)
+	if (const FT_Face* FreeTypeface = &TypefaceFontDataMap[InTypefaceIndex].GetTypeface())
 	{
-		TypefaceFontDataHash = HashCombine(0, GetTypeHash(Typeface.Font));
-		FT_Set_Char_Size(FreeTypeFace, FontSize, FontSize, 96, 96);
-		FT_Set_Pixel_Sizes(FreeTypeFace, FontSize, FontSize);
+		TypefaceFontDataMap[InTypefaceIndex].SetTypefaceFontDataHash(HashCombine(0, GetTypeHash(Typeface.Font)));
+		FT_Set_Char_Size(*FreeTypeface, FontSize, FontSize, 96, 96);
+		FT_Set_Pixel_Sizes(*FreeTypeface, FontSize, FontSize);
 	}
 }
 
-uint32 FCachedFontData::GetTypefaceFontDataHash()
+uint32 FCachedFontData::GetTypefaceFontDataHash(int32 InTypefaceEntryIndex) const
 {
-	return TypefaceFontDataHash;
+	if (TypefaceFontDataMap.Contains(InTypefaceEntryIndex))
+	{
+		return TypefaceFontDataMap[InTypefaceEntryIndex].GetTypefaceFontDataHash();
+	}
+
+	return 0;
 }
 
-TSharedPtr<int32> FCachedFontData::GetCacheCounter()
+TSharedPtr<int32> FCachedFontData::GetCacheCounter(int32 InTypefaceEntryIndex)
 {
-	return CacheCounter;
+	if (TypefaceFontDataMap.Contains(InTypefaceEntryIndex))
+	{
+		return TypefaceFontDataMap[InTypefaceEntryIndex].GetCacheCounter();
+	}
+
+	return nullptr;
 }
 
-TSharedPtr<int32> FCachedFontData::GetMeshesCacheCounter(const FGlyphMeshParameters& Parameters)
+TSharedPtr<int32> FCachedFontData::GetMeshesCacheCounter(const FGlyphMeshParameters& InParameters)
 {
-	const uint32 HashParameters = GetTypeHash(Parameters);
-	FCachedFontMeshes& CachedMeshes = Meshes.FindOrAdd(HashParameters);
+	const uint32 HashParameters = GetTypeHash(InParameters);
+
+	const uint32 TypefaceIndex = InParameters.TypefaceIndex;
+	if (TypefaceFontDataMap.Contains(TypefaceIndex))
+	{
+		FCachedFontMeshes& CachedMeshes = TypefaceFontDataMap[TypefaceIndex].FindOrAddMeshes(HashParameters);
+		return CachedMeshes.GetCacheCounter();
+	}
 
-	return CachedMeshes.GetCacheCounter();
+	return nullptr;
 }
 
-UStaticMesh* FCachedFontData::GetGlyphMesh(uint32 GlyphIndex, const FGlyphMeshParameters& Parameters, const FFreeTypeFace* FontFaceData)
+UStaticMesh* FCachedFontData::GetGlyphMesh(uint32 InGlyphIndex, const FGlyphMeshParameters& InParameters, const FFreeTypeFace* InFontFaceData)
 {
-	const uint32 HashParameters = GetTypeHash(Parameters);
-	FCachedFontMeshes& CachedMeshes = Meshes.FindOrAdd(HashParameters);
+	const uint32 TypefaceIndex = InParameters.TypefaceIndex;
+	const uint32 HashParameters = GetTypeHash(InParameters);
 
-	TObjectPtr<UStaticMesh>* CachedStaticMesh = CachedMeshes.Glyphs.Find(GlyphIndex);
+	if (!TypefaceFontDataMap.Contains(TypefaceIndex))
+	{
+		return nullptr;
+	}
+
+	FCachedFontMeshes& CachedMeshes = TypefaceFontDataMap[TypefaceIndex].FindOrAddMeshes(HashParameters);
+
+	TObjectPtr<UStaticMesh>* CachedStaticMesh = CachedMeshes.Glyphs.Find(InGlyphIndex);
 	if (CachedStaticMesh)
 	{
 		return *CachedStaticMesh;
@@ -208,59 +298,77 @@ UStaticMesh* FCachedFontData::GetGlyphMesh(uint32 GlyphIndex, const FGlyphMeshPa
 
 	uint32 HashGroup = 0;
 	HashGroup = HashCombine(HashGroup, GetTypeHash(Font));
-	HashGroup = HashCombine(HashGroup, GetTypeHash(GlyphIndex));
+	HashGroup = HashCombine(HashGroup, GetTypeHash(InGlyphIndex));
 	const FString StaticMeshName = FString::Printf(TEXT("Text3D_Char_%u_%u"), HashGroup, HashParameters);
 
-	const TSharedContourNode Root = GetGlyphContours(GlyphIndex, FontFaceData);
+	const TSharedContourNode Root = GetGlyphContours(InGlyphIndex, TypefaceIndex, InFontFaceData);
 	if (Root->Children.Num() == 0)
 	{
 		return nullptr;
 	}
 
 	UText3DEngineSubsystem* Subsystem = GEngine->GetEngineSubsystem<UText3DEngineSubsystem>();
-	UStaticMesh* StaticMesh = NewObject<UStaticMesh>(Subsystem, *StaticMeshName);
-	CachedMeshes.Glyphs.Add(GlyphIndex, StaticMesh);
+	const FName StrokeShapeUniqueName = MakeUniqueObjectName(Subsystem, UStaticMesh::StaticClass(), FName(*StaticMeshName));
+
+	UStaticMesh* StaticMesh = NewObject<UStaticMesh>(Subsystem, StrokeShapeUniqueName);
+	CachedMeshes.Glyphs.Add(InGlyphIndex, StaticMesh);
 
 	FMeshCreator MeshCreator;
-	MeshCreator.CreateMeshes(Root, Parameters.Extrude, Parameters.Bevel, Parameters.BevelType, Parameters.BevelSegments, Parameters.bOutline, Parameters.OutlineExpand);
-	MeshCreator.SetFrontAndBevelTextureCoordinates(Parameters.Bevel);
-	MeshCreator.MirrorGroups(Parameters.Extrude);
-		
+	MeshCreator.CreateMeshes(Root, InParameters.Extrude, InParameters.Bevel, InParameters.BevelType, InParameters.BevelSegments, InParameters.bOutline, InParameters.OutlineExpand);
+	MeshCreator.SetFrontAndBevelTextureCoordinates(InParameters.Bevel);
+	MeshCreator.MirrorGroups(InParameters.Extrude);
+
 	MeshCreator.BuildMesh(StaticMesh, Subsystem->DefaultMaterial);
 
 	return StaticMesh;
 }
 
-FT_Face FCachedFontData::GetFreeTypeFace()
-{ 
-	return FreeTypeFace;
+FT_Face FCachedFontData::GetFreeTypeFace(uint32 InTypefaceIndex)
+{
+	if (TypefaceFontDataMap.Contains(InTypefaceIndex))
+	{
+		return TypefaceFontDataMap[InTypefaceIndex].GetTypeface();
+	}
+
+	return nullptr;
 }
 
-const FString& FCachedFontData::GetFontName()
+FString FCachedFontData::GetFontName() const
 {
-	return FontName;
+	if (Font)
+	{
+		return Font->GetName();
+	}
+
+	return TEXT("");
 }
 
-TSharedContourNode FCachedFontData::GetGlyphContours(uint32 GlyphIndex, const FFreeTypeFace* FontFaceData)
+TSharedContourNode FCachedFontData::GetGlyphContours(uint32 InGlyphIndex, int32 InTypefaceEntryIndex, const FFreeTypeFace* InFontFaceData)
 {
-	const FT_Face& GlyphFreeTypeFace = FontFaceData && FontFaceData->IsFaceValid()
-		? FontFaceData->GetFace()
-		: FreeTypeFace;
+	const bool bFontFaceValid = InFontFaceData && InFontFaceData->IsFaceValid();
+	if (!bFontFaceValid && !GetFreeTypeFace(InTypefaceEntryIndex))
+	{
+		return nullptr;
+	}
+
+	const FT_Face& GlyphFreeTypeFace = bFontFaceValid
+		? InFontFaceData->GetFace()
+		: GetFreeTypeFace(InTypefaceEntryIndex);
 	check(GlyphFreeTypeFace);
 
-	if (FT_Load_Glyph(GlyphFreeTypeFace, GlyphIndex, FT_LOAD_NO_HINTING | FT_LOAD_NO_BITMAP))
+	if (FT_Load_Glyph(GlyphFreeTypeFace, InGlyphIndex, FT_LOAD_NO_HINTING | FT_LOAD_NO_BITMAP))
 	{
 		return nullptr;
 	}
 
 #if UE_BUILD_DEBUG
 	// If true, the symbol itself was probably bitmap-only, and this will have placeholder contours
-	bool bIsLastResort = FontFaceData && FontFaceData->GetFace()->family_name
-		? FString(FontFaceData->GetFace()->family_name) == TEXT("LastResort")
+	bool bIsLastResort = InFontFaceData && InFontFaceData->GetFace()->family_name
+		? FString(InFontFaceData->GetFace()->family_name) == TEXT("LastResort")
 		: false;
 	if (bIsLastResort)
 	{
-		UE_LOG(LogText3D, Warning, TEXT("The following glyph index is bitmap-only: %i"), GlyphIndex);
+		UE_LOG(LogText3D, Warning, TEXT("The following glyph index is bitmap-only: %i"), InGlyphIndex);
 	}
 #endif
 
@@ -270,24 +378,61 @@ TSharedContourNode FCachedFontData::GetGlyphContours(uint32 GlyphIndex, const FF
 	return Root;
 }
 
+void FCachedFontData::PrintCache() const
+{
+	UE_LOG(LogText3D, Log, TEXT("\n== Cached typefaces for font: %s =="), *GetFontName());
+
+	for (const TPair<uint32, FTypefaceFontData>& TypefaceData : TypefaceFontDataMap)
+	{
+		UE_LOG(LogText3D, Log, TEXT("\t- %s"), *TypefaceData.Value.GetTypefaceName().ToString());
+	}
+}
+
 bool FCachedFontData::Cleanup()
 {
-	// If there is no Text3D objects using the same font we can just release it all
-	if (CacheCounter.GetSharedReferenceCount() <= 1)
+	TArray<uint32> TypefaceIndexes;
+	TypefaceFontDataMap.GetKeys(TypefaceIndexes);
+
+	for (const uint32 TypefaceIndex : TypefaceIndexes)
+	{
+		CleanupTypeface(TypefaceIndex);
+	}
+
+	if (TypefaceFontDataMap.IsEmpty())
 	{
-		Meshes.Reset();
-		Glyphs.Reset();
-		ClearFreeTypeFace();
 		return true;
 	}
 
-	for (auto It = Meshes.CreateIterator(); It; ++It)
+	return false;
+}
+
+bool FCachedFontData::CleanupTypeface(uint32 InTypefaceIndex)
+{
+	bool bRemoved = false;
+	if (TypefaceFontDataMap.Contains(InTypefaceIndex))
 	{
-		if (It.Value().GetCacheCount() <= 1)
+		FTypefaceFontData& TypefaceFontData = TypefaceFontDataMap[InTypefaceIndex];
+
+		for (auto MeshesIterator = TypefaceFontData.GetMeshes().CreateIterator(); MeshesIterator; ++MeshesIterator)
 		{
-			It.RemoveCurrent();
+			if (MeshesIterator.Value().GetCacheCount() <= 1)
+			{
+				MeshesIterator.RemoveCurrent();
+			}
+		}
+
+		if (TypefaceFontData.GetCacheCount() <= 1)
+		{
+			TypefaceFontData.Reset();
+
+			bRemoved = true;
 		}
 	}
 
-	return false;
+	if (bRemoved)
+	{
+		TypefaceFontDataMap.Remove(InTypefaceIndex);
+	}
+
+	return bRemoved;
 }
diff --git a/Engine/Plugins/Experimental/Text3D/Source/Text3D/Private/Text3DEngineSubsystem.h b/Engine/Plugins/Experimental/Text3D/Source/Text3D/Private/Text3DEngineSubsystem.h
--- a/Engine/Plugins/Experimental/Text3D/Source/Text3D/Private/Text3DEngineSubsystem.h
+++ b/Engine/Plugins/Experimental/Text3D/Source/Text3D/Private/Text3DEngineSubsystem.h
@@ -2,16 +2,17 @@
 
 #pragma once
 
-#include "CoreMinimal.h"
-#include "Subsystems/EngineSubsystem.h"
-#include "Text3DPrivate.h"
 #include "BevelType.h"
-#include "ContourNode.h"
 #include "Containers/Ticker.h"
+#include "ContourNode.h"
+#include "Subsystems/EngineSubsystem.h"
+#include "Text3DPrivate.h"
 
 #include "Text3DEngineSubsystem.generated.h"
 
 class FFreeTypeFace;
+class UFont;
+class UStaticMesh;
 
 USTRUCT()
 struct FGlyphMeshParameters
@@ -36,6 +37,9 @@ public:
 
 	UPROPERTY()
 	float OutlineExpand = 0.5f;
+
+	UPROPERTY()
+	uint32 TypefaceIndex = 0;
 };
 
 inline uint32 GetTypeHash(const FGlyphMeshParameters A)
@@ -47,6 +51,7 @@ inline uint32 GetTypeHash(const FGlyphMeshParameters A)
 	HashParameters = HashCombine(HashParameters, GetTypeHash(A.BevelSegments));
 	HashParameters = HashCombine(HashParameters, GetTypeHash(A.bOutline));
 	HashParameters = HashCombine(HashParameters, GetTypeHash(A.OutlineExpand));
+	HashParameters = HashCombine(HashParameters, GetTypeHash(A.TypefaceIndex));
 	return HashParameters;
 }
 
@@ -58,14 +63,77 @@ struct FCachedFontMeshes
 public:
 	FCachedFontMeshes();
 
-	int32 GetCacheCount();
+	int32 GetCacheCount() const;
 	TSharedPtr<int32> GetCacheCounter();
 
 	UPROPERTY()
-	TMap<uint32, TObjectPtr<class UStaticMesh>> Glyphs;
+	TMap<uint32, TObjectPtr<UStaticMesh>> Glyphs;
+
+private:
+	TSharedPtr<int32> CacheCounter;
+};
+
+USTRUCT()
+struct FTypefaceFontData
+{
+	GENERATED_BODY()
+
+public:
+	FTypefaceFontData();
+
+	void Reset()
+	{
+		FT_Done_Face(FreeTypeFace);
+		FreeTypeFace = nullptr;
+		TypefaceData.Reset();
+	}
+
+	int32 GetCacheCount() const
+	{
+		const int32 Count = CacheCounter.GetSharedReferenceCount();
+		return Count;
+	}
+
+	FT_Face& GetTypeface() { return FreeTypeFace; }
+
+	TArray<uint8>& GetTypefaceData() { return TypefaceData; }
+
+	void SetTypefaceData(const TArray<uint8>& InTypefaceData)
+	{
+		TypefaceData = InTypefaceData;
+	}
+
+	FName GetTypefaceName() const { return TypefaceName; }
+	void SetTypefaceName(FName InTypefaceName)
+	{
+		TypefaceName = InTypefaceName;
+	}
+
+	uint32 GetTypefaceFontDataHash() const { return TypefaceFontDataHash; }
+
+	void SetTypefaceFontDataHash(uint32 InDataHash)
+	{
+		TypefaceFontDataHash = InDataHash;
+	}
+
+	TMap<uint32, FCachedFontMeshes>& GetMeshes() { return Meshes; }
+
+	TSharedPtr<int32> GetCacheCounter() { return CacheCounter; }
+
+	FCachedFontMeshes& FindOrAddMeshes(uint32 InHashParameters)
+	{
+		return Meshes.FindOrAdd(InHashParameters);
+	}
 
 private:
+	UPROPERTY()
+	TMap<uint32, FCachedFontMeshes> Meshes;
+
+	FName TypefaceName;
+	FT_Face FreeTypeFace;
+	TArray<uint8> TypefaceData;
 	TSharedPtr<int32> CacheCounter;
+	uint32 TypefaceFontDataHash;
 };
 
 USTRUCT()
@@ -77,51 +145,62 @@ public:
 	FCachedFontData();
 	~FCachedFontData();
 
-	FT_Face GetFreeTypeFace();
-	const FString& GetFontName();
+	FT_Face GetFreeTypeFace(uint32 InTypefaceIndex);
+
+	FString GetFontName() const;
 
-	void LoadFreeTypeFace();
+	UFont* GetFont() { return Font; }
+
+	void SetFont(UFont* InFont)
+	{
+		Font = InFont;
+	}
+
+	void LoadFreeTypeFace(uint32 InTypefaceIndex);
 	void ClearFreeTypeFace();
 
 	bool Cleanup();
+	bool CleanupTypeface(uint32 InTypefaceIndex);
 
-	uint32 GetTypefaceFontDataHash();
-	TSharedPtr<int32> GetCacheCounter();
-	TSharedPtr<int32> GetMeshesCacheCounter(const FGlyphMeshParameters& Parameters);
+	uint32 GetTypefaceFontDataHash(int32 InTypefaceEntryIndex) const;
+	TSharedPtr<int32> GetCacheCounter(int32 InTypefaceEntryIndex);
+	TSharedPtr<int32> GetMeshesCacheCounter(const FGlyphMeshParameters& InParameters);
 
-	UStaticMesh* GetGlyphMesh(uint32 GlyphIndex, const FGlyphMeshParameters& Parameters, const FFreeTypeFace* FontFaceData = nullptr);
-	TSharedContourNode GetGlyphContours(uint32 GlyphIndex, const FFreeTypeFace* FontFaceData = nullptr);
+	UStaticMesh* GetGlyphMesh(uint32 InGlyphIndex, const FGlyphMeshParameters& InParameters, const FFreeTypeFace* InFontFaceData = nullptr);
+	TSharedContourNode GetGlyphContours(uint32 InGlyphIndex, int32 InTypefaceEntryIndex, const FFreeTypeFace* InFontFaceData = nullptr);
 
-	UPROPERTY()
-	TObjectPtr<class UFont> Font;
+	void PrintCache() const;
 
+private:
 	UPROPERTY()
-	TMap<uint32, FCachedFontMeshes> Meshes;
+	TObjectPtr<UFont> Font;
 
-	TMap<uint32, TSharedContourNode> Glyphs;
-
-private:
-	FT_Face FreeTypeFace;
-	FString FontName;
-	TArray<uint8> Data;
-	TSharedPtr<int32> CacheCounter;
-	uint32 TypefaceFontDataHash;
+	UPROPERTY()
+	TMap<uint32, FTypefaceFontData> TypefaceFontDataMap;
 };
 
 UCLASS()
-class TEXT3D_API UText3DEngineSubsystem : public UEngineSubsystem
+class TEXT3D_API UText3DEngineSubsystem : public UEngineSubsystem, public FSelfRegisteringExec
 {
 	GENERATED_BODY()
 
-public:	
+public:
 	UText3DEngineSubsystem();
 
+	// ~Begin UEngineSubsystem
 	virtual void Initialize(FSubsystemCollectionBase& Collection) override;
 	virtual void Deinitialize() override;
+	// ~End UEngineSubsystem
+
+	// ~Begin FSelfRegisteringExec
+	virtual bool Exec(class UWorld* InWorld, const TCHAR* InCmd, FOutputDevice& InAr) override;
+	// ~End FSelfRegisteringExec
+
+	void PrintCache() const;
 
 	void Reset();
 	void Cleanup();
-	FCachedFontData& GetCachedFontData(class UFont* Font);
+	FCachedFontData& GetCachedFontData(UFont* InFont, int32 InTypefaceEntryIndex = 0);
 
 	UPROPERTY()
 	TObjectPtr<class UMaterial> DefaultMaterial;
diff --git a/Engine/Plugins/Experimental/Text3D/Source/Text3D/Public/Text3DComponent.h b/Engine/Plugins/Experimental/Text3D/Source/Text3D/Public/Text3DComponent.h
--- a/Engine/Plugins/Experimental/Text3D/Source/Text3D/Public/Text3DComponent.h
+++ b/Engine/Plugins/Experimental/Text3D/Source/Text3D/Public/Text3DComponent.h
@@ -14,6 +14,7 @@ class FTextLayout;
 class ITextLayoutMarshaller;
 class UFont;
 class UMaterialInterface;
+struct FTypefaceEntry;
 
 UENUM()
 enum class EText3DVerticalTextAlignment : uint8
@@ -39,7 +40,7 @@ enum class EText3DModifyFlags : uint8
 	Layout = 1 << 0,
 	Geometry = 1 << 1,
 	Unfreeze = 1 << 2,
-    
+
 	All = Layout | Geometry | Unfreeze
 };
 ENUM_CLASS_FLAGS(EText3DModifyFlags)
@@ -269,11 +270,18 @@ public:
 	/** Gets all the glyph meshes */
 	UFUNCTION(BlueprintCallable, Category = "Rendering|Components|Text3D")
 	const TArray<UStaticMeshComponent*>& GetGlyphMeshComponents();
-	
+
 	/** Gets the scale of actual text geometry, taking into account MaxWidth and MaxHeight constraints. This function will NOT return the component scale*/
 	UFUNCTION(BlueprintCallable, Category = "Rendering|Components|Text3D")
 	FVector GetTextScale();
 
+
+	/** Get the typeface */
+	FName GetTypeface() const { return  Typeface; }
+
+	/** Set the typeface */
+	void SetTypeface(const FName InTypeface);
+
 	/** Manually update the geometry, ignoring RefreshOnChange (but still accounting for the Freeze flag) */
 	void Rebuild();
 
@@ -330,6 +338,9 @@ protected:
 	UPROPERTY(EditAnywhere, BlueprintReadWrite, Getter, Setter, Category = "Text", meta = (AllowPrivateAccess = "true"))
 	TObjectPtr<UFont> Font;
 
+	UPROPERTY(EditAnywhere, BlueprintReadWrite, Getter, Setter, Category = "Text", meta = (AllowPrivateAccess = "true", GetOptions="GetTypefaceNames"))
+	FName Typeface;
+
 	/** Horizontal text alignment */
 	UPROPERTY(EditAnywhere, BlueprintReadWrite, Getter, Setter, Category = "Layout", meta = (AllowPrivateAccess = "true"))
 	EText3DHorizontalTextAlignment HorizontalAlignment;
@@ -371,7 +382,7 @@ protected:
 	bool bScaleProportionally;
 
 	// Lighting flags
-	
+
 	/** Controls whether the text glyphs should cast a shadow or not. */
 	UPROPERTY(EditAnywhere, BlueprintReadWrite, Getter = "CastsShadow", Setter = "SetCastShadow", Category = "Lighting", meta = (AllowPrivateAccess = "true"))
 	bool bCastShadow = true;
@@ -382,8 +393,11 @@ protected:
 	 */
 	UFUNCTION(BlueprintCallable, Category = "Text3D")
 	FText GetFormattedText() const;
-	
+
 protected:
+	UFUNCTION()
+	TArray<FName> GetTypefaceNames() const;
+
 	/** Intercept and propagate a change on this component to all children. */
 	virtual void OnVisibilityChanged() override;
 
@@ -474,4 +488,9 @@ private:
 	void MarkForGeometryUpdate();
 	void MarkForLayoutUpdate();
 	void ClearUpdateFlags();
+
+	uint32 GetTypeFaceIndex() const;
+	bool IsTypefaceAvailable(FName InTypeface) const;
+	TArray<FTypefaceEntry> GetAvailableTypefaces() const;
+	void RefreshTypeface();
 };
