[Motion Design] Text 3D Actor/Component: Font geometry is not properly updated after changing alignment layout

diff --git a/Engine/Plugins/Experimental/Text3D/Source/Text3D/Private/Text3DComponent.cpp b/Engine/Plugins/Experimental/Text3D/Source/Text3D/Private/Text3DComponent.cpp
--- a/Engine/Plugins/Experimental/Text3D/Source/Text3D/Private/Text3DComponent.cpp
+++ b/Engine/Plugins/Experimental/Text3D/Source/Text3D/Private/Text3DComponent.cpp
@@ -1053,7 +1053,7 @@ void UText3DComponent::UpdateTransforms()
 		}
 	}
 
-	ModifyFlags &= EText3DModifyFlags::Layout; 
+	ModifyFlags &= ~EText3DModifyFlags::Layout;
 }
 
 void UText3DComponent::ClearTextMesh()
