[Text3D] SlateApplication check to prevent crash on cook/command line

diff --git a/Engine/Plugins/Experimental/Text3D/Source/Text3D/Private/TextShaper.cpp b/Engine/Plugins/Experimental/Text3D/Source/Text3D/Private/TextShaper.cpp
--- a/Engine/Plugins/Experimental/Text3D/Source/Text3D/Private/TextShaper.cpp
+++ b/Engine/Plugins/Experimental/Text3D/Source/Text3D/Private/TextShaper.cpp
@@ -6,6 +6,7 @@
 #include "Containers/Array.h"
 #include "Containers/UnrealString.h"
 #include "Fonts/FontCache.h"
+#include "Framework/Application/SlateApplication.h"
 #include "Framework/Text/PlainTextLayoutMarshaller.h"
 #include "Framework/Text/ShapedTextCache.h"
 #include "Framework/Text/SlateTextRun.h"
@@ -38,6 +39,13 @@ void FTextShaper::ShapeBidirectionalText(
 	const TSharedPtr<ITextLayoutMarshaller>& TextMarshaller,
 	TArray<FShapedGlyphLine>& OutShapedLines)
 {
+	// @todo: Restore when dependency on SlateApplication is removed. Currently this means text meshes won't be created on the server.
+	// @see: UE-211843
+	if (!FSlateApplication::IsInitialized())
+	{
+		return;
+	}
+
 	TextLayout->ClearLines();
 	TextLayout->ClearLineHighlights();
 	TextLayout->ClearRunRenderers();
