MotionDesign : Text3D
- Fixed materials slot being created even though no data contained in the polygroup
- Simplifying the material slots index search using named slots
- Back will not be created by default if no extrude amount is set

diff --git a/Engine/Plugins/Experimental/Text3D/Source/Text3D/Private/Glyph.cpp b/Engine/Plugins/Experimental/Text3D/Source/Text3D/Private/Glyph.cpp
--- a/Engine/Plugins/Experimental/Text3D/Source/Text3D/Private/Glyph.cpp
+++ b/Engine/Plugins/Experimental/Text3D/Source/Text3D/Private/Glyph.cpp
@@ -21,9 +21,32 @@ void FText3DGlyph::Build(UStaticMesh* StaticMesh, UMaterial* DefaultMaterial)
 {
 	check(StaticMesh);
 
+	auto AddMaterial = [](UStaticMesh* InMesh, UMaterial* InMaterial, int32 InIndex)
+	{
+		using namespace UE::Text3D::Materials;
+
+		check(SlotNames.IsValidIndex(InIndex))
+
+		FName MaterialName = SlotNames[InIndex];
+
+#if WITH_EDITORONLY_DATA
+		FStaticMaterial& StaticMaterial = InMesh->GetStaticMaterials().Emplace_GetRef(InMaterial, MaterialName, MaterialName);
+#else
+		FStaticMaterial& StaticMaterial = InMesh->GetStaticMaterials().Emplace_GetRef(InMaterial, MaterialName);
+#endif
+
+		StaticMaterial.UVChannelData = FMeshUVChannelInfo(1.0f);
+
+		return MaterialName;
+	};
+
 	for (int32 Index = 0; Index < Groups.Num(); Index++)
 	{
-		StaticMeshAttributes.GetPolygonGroupMaterialSlotNames()[FPolygonGroupID(Index)] = StaticMesh->AddMaterial(DefaultMaterial);
+		const FPolygonGroupID PolyGroup(Index);
+		if (MeshDescription.GetNumPolygonGroupTriangles(PolyGroup) != 0)
+		{
+			StaticMeshAttributes.GetPolygonGroupMaterialSlotNames()[PolyGroup] = AddMaterial(StaticMesh, DefaultMaterial, Index);
+		}
 	}
 
 	TArray<const FMeshDescription*> MeshDescriptions;
diff --git a/Engine/Plugins/Experimental/Text3D/Source/Text3D/Private/MeshCreator.cpp b/Engine/Plugins/Experimental/Text3D/Source/Text3D/Private/MeshCreator.cpp
--- a/Engine/Plugins/Experimental/Text3D/Source/Text3D/Private/MeshCreator.cpp
+++ b/Engine/Plugins/Experimental/Text3D/Source/Text3D/Private/MeshCreator.cpp
@@ -78,13 +78,21 @@ void FMeshCreator::SetFrontAndBevelTextureCoordinates(const float Bevel)
 	};
 
 	SetTextureCoordinates(EText3DGroupType::Front);
-	SetTextureCoordinates(EText3DGroupType::Bevel);
+
+	if (!FMath::IsNearlyZero(Bevel))
+	{
+		SetTextureCoordinates(EText3DGroupType::Bevel);
+	}
 }
 
 void FMeshCreator::MirrorGroups(const float Extrude)
 {
-	MirrorGroup(EText3DGroupType::Front, EText3DGroupType::Back, Extrude);
-	MirrorGroup(EText3DGroupType::Bevel, EText3DGroupType::Bevel, Extrude);
+	// No bevels without extrude
+	if (!FMath::IsNearlyZero(Extrude))
+	{
+		MirrorGroup(EText3DGroupType::Front, EText3DGroupType::Back, Extrude);
+		MirrorGroup(EText3DGroupType::Bevel, EText3DGroupType::Bevel, Extrude);
+	}
 }
 
 void FMeshCreator::BuildMesh(UStaticMesh* StaticMesh, class UMaterial* DefaultMaterial)
@@ -176,8 +184,13 @@ void FMeshCreator::CreateBevelMesh(const float Bevel, const EText3DBevelType Typ
 
 void FMeshCreator::CreateExtrudeMesh(float Extrude, float Bevel, const EText3DBevelType Type, bool bFlipNormals)
 {
+	if (FMath::IsNearlyZero(Extrude))
+	{
+		return;
+	}
+
 	Bevel = FMath::Max(UE_SMALL_NUMBER, Bevel);
-	
+
 	if (Type != EText3DBevelType::HalfCircle)
 	{
 		Bevel = FMath::Clamp(Bevel, 0.0f, Extrude / 2.f);
diff --git a/Engine/Plugins/Experimental/Text3D/Source/Text3D/Private/Text3DComponent.cpp b/Engine/Plugins/Experimental/Text3D/Source/Text3D/Private/Text3DComponent.cpp
--- a/Engine/Plugins/Experimental/Text3D/Source/Text3D/Private/Text3DComponent.cpp
+++ b/Engine/Plugins/Experimental/Text3D/Source/Text3D/Private/Text3DComponent.cpp
@@ -1309,44 +1309,20 @@ float UText3DComponent::MaxBevel() const
 
 void UText3DComponent::OnMaterialChanged()
 {
-	// Material indices are affected by some options and can differ
-	auto GetMaterialTypeIndex = [this](EText3DGroupType InType)->int32
-	{
-		const bool bHasBevel = !bOutline && !FMath::IsNearlyZero(Bevel);
-
-		int32 MaterialIndex = static_cast<int32>(InType);
-
-		if (InType >= EText3DGroupType::Extrude && !bHasBevel)
-		{
-			MaterialIndex -= 1;
-		}
-
-		return MaterialIndex;
-	};
+	using namespace UE::Text3D::Materials;
 
-	/**
-	* GetNumMaterials should return 4 but the index of the material changes based on options used
-	* [Front, (Extrude|Bevel), Back] if no extrude or bevel
-	* [Front, Extrude, Back] if no bevel since bevel only works when there is extrude anyway
-	* [Front, Bevel, Extrude, Back] if extrude and bevel is set
-	 */
 	for (UStaticMeshComponent* StaticMeshComponent : CharacterMeshes)
 	{
-		const int32 MaterialCount = StaticMeshComponent->GetNumMaterials();
-
-		for (int32 MaterialIndex = 0; MaterialIndex < MaterialCount; MaterialIndex++)
+		for (int32 GroupIndex = 0; GroupIndex < static_cast<int32>(EText3DGroupType::TypeCount); GroupIndex++)
 		{
-			// Should not have any material slot above the back material slot but in case set them as nullptr
-			if (MaterialIndex >= static_cast<int32>(EText3DGroupType::TypeCount))
+			const int32 MaterialIndex = StaticMeshComponent->GetMaterialIndex(SlotNames[GroupIndex]);
+
+			if (MaterialIndex == INDEX_NONE)
 			{
-				StaticMeshComponent->SetMaterial(MaterialIndex, nullptr);
 				continue;
 			}
 
-			// Get the material type from the current index and based on text options adapt the index to assign proper material
-			const EText3DGroupType MaterialType = static_cast<EText3DGroupType>(MaterialIndex);
-			const int32 MaterialTypeIndex = GetMaterialTypeIndex(MaterialType);
-			StaticMeshComponent->SetMaterial(MaterialTypeIndex, GetMaterial(MaterialType));
+			StaticMeshComponent->SetMaterial(MaterialIndex, GetMaterial(static_cast<EText3DGroupType>(GroupIndex)));
 		}
 	}
 }
diff --git a/Engine/Plugins/Experimental/Text3D/Source/Text3D/Public/Mesh.h b/Engine/Plugins/Experimental/Text3D/Source/Text3D/Public/Mesh.h
--- a/Engine/Plugins/Experimental/Text3D/Source/Text3D/Public/Mesh.h
+++ b/Engine/Plugins/Experimental/Text3D/Source/Text3D/Public/Mesh.h
@@ -38,3 +38,14 @@ enum class EText3DGroupType : uint8
 };
 
 using TText3DMeshList = TArray<FText3DMesh, TFixedAllocator<static_cast<int32>(EText3DGroupType::TypeCount)>>;
+
+namespace UE::Text3D::Materials
+{
+	static const TArray<FName> SlotNames =
+	{
+		TEXT("Front"),
+		TEXT("Bevel"),
+		TEXT("Extrude"),
+		TEXT("Back")
+	};
+}
diff --git a/Engine/Plugins/Experimental/Text3D/Source/Text3D/Public/Text3DComponent.h b/Engine/Plugins/Experimental/Text3D/Source/Text3D/Public/Text3DComponent.h
--- a/Engine/Plugins/Experimental/Text3D/Source/Text3D/Public/Text3DComponent.h
+++ b/Engine/Plugins/Experimental/Text3D/Source/Text3D/Public/Text3DComponent.h
@@ -299,15 +299,15 @@ protected:
 	float Extrude;
 
 	/** Size of bevel */
-	UPROPERTY(EditAnywhere, BlueprintReadWrite, Getter, Setter, Category = "Geometry", meta = (EditCondition = "!bOutline", ClampMin = 0, AllowPrivateAccess = "true"))
+	UPROPERTY(EditAnywhere, BlueprintReadWrite, Getter, Setter, Category = "Geometry", meta = (EditCondition = "!bOutline && Extrude > 0", ClampMin = 0, AllowPrivateAccess = "true"))
 	float Bevel;
 
 	/** Bevel Type */
-	UPROPERTY(EditAnywhere, BlueprintReadWrite, Getter, Setter, Category = "Geometry", meta = (EditCondition = "!bOutline", AllowPrivateAccess = "true"))
+	UPROPERTY(EditAnywhere, BlueprintReadWrite, Getter, Setter, Category = "Geometry", meta = (EditCondition = "!bOutline && Extrude > 0", AllowPrivateAccess = "true"))
 	EText3DBevelType BevelType;
 
 	/** Bevel Segments (Defines the amount of tesselation for the bevel part) */
-	UPROPERTY(EditAnywhere, BlueprintReadWrite, Getter, Setter, Category = "Geometry", meta = (EditCondition = "!bOutline", ClampMin = 1, ClampMax = 15, AllowPrivateAccess = "true"))
+	UPROPERTY(EditAnywhere, BlueprintReadWrite, Getter, Setter, Category = "Geometry", meta = (EditCondition = "!bOutline && Extrude > 0", ClampMin = 1, ClampMax = 15, AllowPrivateAccess = "true"))
 	int32 BevelSegments;
 
 	/** Generate Outline */
@@ -323,15 +323,15 @@ protected:
 	TObjectPtr<UMaterialInterface> FrontMaterial;
 
 	/** Material for the bevel part */
-	UPROPERTY(EditAnywhere, BlueprintReadWrite, Getter, Setter, Category = "Materials", meta = (EditCondition = "!bOutline", AllowPrivateAccess = "true"))
+	UPROPERTY(EditAnywhere, BlueprintReadWrite, Getter, Setter, Category = "Materials", meta = (EditCondition = "!bOutline && Extrude > 0 && Bevel > 0", EditConditionHides, AllowPrivateAccess = "true"))
 	TObjectPtr<UMaterialInterface> BevelMaterial;
 
 	/** Material for the extruded part */
-	UPROPERTY(EditAnywhere, BlueprintReadWrite, Getter, Setter, Category = "Materials", meta = (AllowPrivateAccess = "true"))
+	UPROPERTY(EditAnywhere, BlueprintReadWrite, Getter, Setter, Category = "Materials", meta = (EditCondition = "Extrude > 0", EditConditionHides, AllowPrivateAccess = "true"))
 	TObjectPtr<UMaterialInterface> ExtrudeMaterial;
 
 	/** Material for the back part */
-	UPROPERTY(EditAnywhere, BlueprintReadWrite, Getter, Setter, Category = "Materials", meta = (AllowPrivateAccess = "true"))
+	UPROPERTY(EditAnywhere, BlueprintReadWrite, Getter, Setter, Category = "Materials", meta = (EditCondition = "Extrude > 0", EditConditionHides, AllowPrivateAccess = "true"))
 	TObjectPtr<UMaterialInterface> BackMaterial;
 
 	/** Text font */
