MotionDesign : Text3D
- Adding the ability to wrap 3D text in the case where it exceeds max width (Github change from chris-visionsof)

diff --git a/Engine/Plugins/Experimental/Text3D/Source/Text3D/Private/Text3DComponent.cpp b/Engine/Plugins/Experimental/Text3D/Source/Text3D/Private/Text3DComponent.cpp
--- a/Engine/Plugins/Experimental/Text3D/Source/Text3D/Private/Text3DComponent.cpp
+++ b/Engine/Plugins/Experimental/Text3D/Source/Text3D/Private/Text3DComponent.cpp
@@ -36,12 +36,78 @@ struct FText3DShapedText
 		LineHeight = 0.0f;
 		FontAscender = 0.0f;
 		FontDescender = 0.0f;
+		Kerning = 0.0f;
+		WordSpacing = 0.0f;
+		bWrap = false;
 		Lines.Reset();
 	}
 
+	void CalculateWidth()
+	{
+		TArray<FShapedGlyphLine> NewLines;
+		for (const FShapedGlyphLine& GlyphLine : Lines)
+		{
+			FShapedGlyphLine& CurrentLine = NewLines.Add_GetRef(FShapedGlyphLine());
+			TArray<FShapedGlyphEntry> CurrentWord;
+
+			float LineWidth = 0.0f;
+			int32 GlyphCount = GlyphLine.GlyphsToRender.Num();
+			float CurrentWordLength = 0.0f;
+			for (int32 GlyphIdx = 0; GlyphIdx < GlyphCount; ++GlyphIdx)
+			{
+				const FShapedGlyphEntry& CurrentGlyph = GlyphLine.GlyphsToRender[GlyphIdx];
+				float GlyphAdv = GlyphLine.GetAdvance(GlyphIdx, Kerning, WordSpacing);
+				LineWidth += GlyphAdv;
+				CurrentWordLength += GlyphAdv;
+
+				// If we're at the end the line or at whitespace
+				if (!CurrentGlyph.bIsVisible || GlyphIdx == GlyphLine.GlyphsToRender.Num() - 1)
+				{
+					if (bWrap && LineWidth > MaxWidth)
+					{
+						CurrentLine.Width = LineWidth - CurrentWordLength;
+
+						if (CurrentWordLength != LineWidth) // No break to wrap
+						{
+							NewLines.Add(FShapedGlyphLine());
+							LineWidth = CurrentWordLength;
+						}
+					}
+					else
+					{
+						CurrentLine.Width = LineWidth;
+					}
+
+					// if we're wrapping, we may/may not want the white space
+					if (!bWrap || LineWidth < MaxWidth || LineWidth == CurrentWordLength)
+					{
+						CurrentWord.Add(CurrentGlyph);
+					}
+
+					NewLines.Last().GlyphsToRender.Append(CurrentWord);
+					CurrentWordLength = 0.0f;
+					CurrentWord.Empty();
+				}
+				else
+				{
+					CurrentWord.Add(CurrentGlyph);
+				}
+			}
+
+			CurrentLine.GlyphsToRender.Append(CurrentWord);
+			CurrentLine.Width = LineWidth;
+		}
+
+		Lines = NewLines;
+	}
+
 	float LineHeight;
 	float FontAscender;
 	float FontDescender;
+	float Kerning;
+	float WordSpacing;
+	float MaxWidth;
+	bool bWrap;
 	TArray<struct FShapedGlyphLine> Lines;
 };
 
@@ -222,7 +288,7 @@ void UText3DComponent::OnRegister()
 		TextRoot->AttachToComponent(this, FAttachmentTransformRules::KeepRelativeTransform);
 	}
 
-	RebuildInternal();
+	RebuildInternal(/** AutoUpdate */true, /** CleanCache */true);
 }
 
 void UText3DComponent::OnUnregister()
@@ -297,11 +363,20 @@ void UText3DComponent::PostEditChangeProperty(FPropertyChangedEvent& PropertyCha
 			 Name == GET_MEMBER_NAME_CHECKED(UText3DComponent, WordSpacing) ||
 			 Name == GET_MEMBER_NAME_CHECKED(UText3DComponent, bHasMaxWidth) ||
 			 Name == GET_MEMBER_NAME_CHECKED(UText3DComponent, MaxWidth) ||
+			 Name == GET_MEMBER_NAME_CHECKED(UText3DComponent, MaxWidthHandling) ||
 			 Name == GET_MEMBER_NAME_CHECKED(UText3DComponent, bHasMaxHeight) ||
 			 Name == GET_MEMBER_NAME_CHECKED(UText3DComponent, MaxHeight) ||
 			 Name == GET_MEMBER_NAME_CHECKED(UText3DComponent, bScaleProportionally))
 	{
-		MarkForLayoutUpdate();
+		if (MaxWidthHandling == EText3DMaxWidthHandling::WrapAndScale ||
+ 		   Name == GET_MEMBER_NAME_CHECKED(UText3DComponent, MaxWidthHandling))
+		{
+			MarkForGeometryUpdate();
+		}
+		else
+		{
+			MarkForLayoutUpdate();
+		}
 	}
 	else if (Name == GET_MEMBER_NAME_CHECKED(UText3DComponent, FrontMaterial)
 		|| Name == GET_MEMBER_NAME_CHECKED(UText3DComponent, BevelMaterial)
@@ -747,6 +822,30 @@ void UText3DComponent::SetMaxWidth(const float Value)
 	}
 }
 
+EText3DMaxWidthHandling UText3DComponent::GetMaxWidthHandling() const
+{
+	return MaxWidthHandling;
+}
+
+void UText3DComponent::SetMaxWidthHandling(const EText3DMaxWidthHandling Value)
+{
+	if (MaxWidthHandling == Value)
+	{
+		return;
+	}
+
+	MaxWidthHandling = Value;
+
+	if (MaxWidthHandling == EText3DMaxWidthHandling::WrapAndScale)
+	{
+		MarkForGeometryUpdate();
+	}
+	else
+	{
+		MarkForLayoutUpdate();
+	}
+}
+
 bool UText3DComponent::HasMaxHeight() const
 {
 	return bHasMaxHeight;
@@ -900,14 +999,6 @@ void UText3DComponent::RebuildInternal(const bool& bIsAutoUpdate, const bool& bC
 	}
 }
 
-void UText3DComponent::CalculateTextWidth()
-{
-	for (FShapedGlyphLine& ShapedLine : ShapedText->Lines)
-	{
-		ShapedLine.CalculateWidth(Kerning, WordSpacing);
-	}
-}
-
 float UText3DComponent::GetTextHeight() const
 {
 	return ShapedText->Lines.Num() * ShapedText->LineHeight + (ShapedText->Lines.Num() - 1) * LineSpacing;
@@ -1002,7 +1093,7 @@ FVector UText3DComponent::GetLineLocation(int32 LineIndex)
 
 void UText3DComponent::UpdateTransforms()
 {
-	CalculateTextWidth();
+	ShapedText->CalculateWidth();
 	CalculateTextScale();
 	const FVector Scale = GetTextScale();
 	TextRoot->SetRelativeScale3D(Scale);
@@ -1071,7 +1162,7 @@ void UText3DComponent::ClearTextMesh()
 			if (IsValid(ChildComponent))
 			{
 				ChildComponent->DetachFromComponent(FDetachmentTransformRules::KeepRelativeTransform);
-				ChildComponent->DestroyComponent();	
+				ChildComponent->DestroyComponent();
 			}
 		}
 	}
@@ -1111,7 +1202,7 @@ void UText3DComponent::BuildTextMesh(const bool& bCleanCache)
 		{
 			if (!UE::IsSavingPackage(StrongThis))
 			{
-				StrongThis->BuildTextMeshInternal(bCleanCache);	
+				StrongThis->BuildTextMeshInternal(bCleanCache);
 			}
 		}
 	});
@@ -1142,7 +1233,7 @@ void UText3DComponent::BuildTextMeshInternal(const bool& bCleanCache)
 	FCachedFontData& CachedFontData = Subsystem->GetCachedFontData(Font, TypefaceIndex);
 	const FT_Face Face = CachedFontData.GetFreeTypeFace(TypefaceIndex);
 	if (!Face)
-	{ 
+	{
 		UE_LOG(LogText3D, Error, TEXT("Failed to load font data '%s'"), *CachedFontData.GetFontName());
 		return;
 	}
@@ -1155,6 +1246,10 @@ void UText3DComponent::BuildTextMeshInternal(const bool& bCleanCache)
 	ShapedText->LineHeight = Face->size->metrics.height * FontInverseScale;
 	ShapedText->FontAscender = Face->size->metrics.ascender * FontInverseScale;
 	ShapedText->FontDescender = Face->size->metrics.descender * FontInverseScale;
+	ShapedText->Kerning = Kerning;
+	ShapedText->WordSpacing = WordSpacing;
+	ShapedText->MaxWidth = MaxWidth;
+	ShapedText->bWrap = MaxWidthHandling == EText3DMaxWidthHandling::WrapAndScale;
 
 	constexpr int32 AdjustedFontSize = 48; // Magic number that makes font scale consistent with previous implementation
 	FSlateFontInfo FontInfo(Font, AdjustedFontSize);
@@ -1206,12 +1301,12 @@ void UText3DComponent::BuildTextMeshInternal(const bool& bCleanCache)
 			if (const TSharedPtr<FFreeTypeFace> FontFacePtr = GlyphEntry.FontFaceData->FontFace.Pin();
 				FontFacePtr.IsValid())
 			{
-				GlyphIndexToFontFace.FindOrAdd(GlyphEntry.GlyphIndex, FontFacePtr.Get());	
+				GlyphIndexToFontFace.FindOrAdd(GlyphEntry.GlyphIndex, FontFacePtr.Get());
 			}
 		}
 	}
 
-	CalculateTextWidth();
+	ShapedText->CalculateWidth();
 	CalculateTextScale();
 	TextRoot->SetRelativeScale3D(GetTextScale());
 
@@ -1262,7 +1357,7 @@ void UText3DComponent::BuildTextMeshInternal(const bool& bCleanCache)
 			}
 			else
 			{
-				// @note: This shouldn't occur, but it does under unknown circumstances (UE-164789) so it should be handled 
+				// @note: This shouldn't occur, but it does under unknown circumstances (UE-164789) so it should be handled
 				UE_LOG(LogText3D, Error, TEXT("CharacterMesh not found at index %d"), GlyphId);
 			}
 
@@ -1381,4 +1476,4 @@ void UText3DComponent::GetBounds(FVector& Origin, FVector& BoxExtent)
 	Box.GetCenterAndExtents(Origin, BoxExtent);
 }
 
-#undef LOCTEXT_NAMESPACE
+#undef LOCTEXT_NAMESPACE
\ No newline at end of file
diff --git a/Engine/Plugins/Experimental/Text3D/Source/Text3D/Public/Text3DComponent.h b/Engine/Plugins/Experimental/Text3D/Source/Text3D/Public/Text3DComponent.h
--- a/Engine/Plugins/Experimental/Text3D/Source/Text3D/Public/Text3DComponent.h
+++ b/Engine/Plugins/Experimental/Text3D/Source/Text3D/Public/Text3DComponent.h
@@ -3,7 +3,6 @@
 #pragma once
 
 #include "Components/SceneComponent.h"
-#include "CoreMinimal.h"
 #include "BevelType.h"
 #include "Mesh.h"
 #include "UObject/ObjectMacros.h"
@@ -33,6 +32,15 @@ enum class EText3DHorizontalTextAlignment : uint8
 	Right			UMETA(DisplayName = "Right"),
 };
 
+UENUM()
+enum class EText3DMaxWidthHandling : uint8
+{
+	/** Scales the text to meet the max width */
+	Scale			UMETA(DisplayName = "Scale"),
+	/** First wraps the text (if possible) and then scales to meet the max width */
+	WrapAndScale	UMETA(DisplayName = "Wrap and Scale"),
+};
+
 UENUM()
 enum class EText3DModifyFlags : uint8
 {
@@ -213,6 +221,13 @@ public:
 	UFUNCTION(BlueprintCallable, Category = "Rendering|Components|Text3D", meta = (DeprecatedFunction, DeprecationMessage = "Set the property directly"))
 	void SetMaxWidth(const float Value);
 
+	/** Get the Maximum Width Handling - Whether to wrap before scaling when the text size reaches the max width */
+	EText3DMaxWidthHandling GetMaxWidthHandling() const;
+
+	/** Set the Maximum Width Handling - Whether to wrap before scaling when the text size reaches the max width */
+	UFUNCTION(BlueprintCallable, Category = "Rendering|Components|Text3D", meta = (DeprecatedFunction, DeprecationMessage = "Set the property directly"))
+	void SetMaxWidthHandling(const EText3DMaxWidthHandling Value);
+
 	/** Whether a maximum height is specified */
 	bool HasMaxHeight() const;
 
@@ -275,7 +290,6 @@ public:
 	UFUNCTION(BlueprintCallable, Category = "Rendering|Components|Text3D")
 	FVector GetTextScale();
 
-
 	/** Get the typeface */
 	FName GetTypeface() const { return  Typeface; }
 
@@ -373,6 +387,10 @@ protected:
 	UPROPERTY(EditAnywhere, BlueprintReadWrite, Getter = "HasMaxHeight", Setter = "SetHasMaxHeight", Category = "Layout", meta = (InlineEditConditionToggle, AllowPrivateAccess = "true"))
 	bool bHasMaxHeight;
 
+	/** Dictates how to handle the text if it exceeds the max width */
+	UPROPERTY(EditAnywhere, BlueprintReadWrite, Getter, Setter, Category = "Layout", meta = (HideEditConditionToggle, EditCondition="bHasMaxWidth", AllowPrivateAccess = "true"))
+	EText3DMaxWidthHandling MaxWidthHandling;
+
 	/** Sets a maximum height to the 3D Text */
 	UPROPERTY(EditAnywhere, BlueprintReadWrite, Getter, Setter, Category = "Layout", meta = (EditCondition = "bHasMaxHeight", ClampMin = 1, AllowPrivateAccess = "true"))
 	float MaxHeight;
@@ -394,7 +412,6 @@ protected:
 	UFUNCTION(BlueprintCallable, Category = "Text3D")
 	FText GetFormattedText() const;
 
-protected:
 	UFUNCTION()
 	TArray<FName> GetTypefaceNames() const;
 
@@ -477,7 +494,6 @@ private:
 	void BuildTextMeshInternal(const bool& bCleanCache);
 
 	/** Layout functionality. */
-	void CalculateTextWidth();
 	float GetTextHeight() const;
 	void CalculateTextScale();
 	FVector GetLineLocation(int32 LineIndex);
@@ -493,4 +509,4 @@ private:
 	bool IsTypefaceAvailable(FName InTypeface) const;
 	TArray<FTypefaceEntry> GetAvailableTypefaces() const;
 	void RefreshTypeface();
-};
+};
\ No newline at end of file
