MotionDesign : Text3D
- Fixed right/center alignment is anchoring to the left when adjusting kerning

diff --git a/Engine/Plugins/Experimental/Text3D/Source/Text3D/Private/Text3DComponent.cpp b/Engine/Plugins/Experimental/Text3D/Source/Text3D/Private/Text3DComponent.cpp
--- a/Engine/Plugins/Experimental/Text3D/Source/Text3D/Private/Text3DComponent.cpp
+++ b/Engine/Plugins/Experimental/Text3D/Source/Text3D/Private/Text3DComponent.cpp
@@ -1127,6 +1127,11 @@ FVector UText3DComponent::GetLineLocation(int32 LineIndex)
 
 void UText3DComponent::UpdateTransforms()
 {
+	ShapedText->Kerning = Kerning;
+	ShapedText->WordSpacing = WordSpacing;
+	ShapedText->MaxWidth = MaxWidth;
+	ShapedText->bWrap = MaxWidthHandling == EText3DMaxWidthHandling::WrapAndScale;
+
 	ShapedText->CalculateWidth();
 	CalculateTextScale();
 	const FVector Scale = GetTextScale();
