MotionDesign : Text3D
- Fix WrapAndScale not working properly with layout alignments

diff --git a/Engine/Plugins/Experimental/Text3D/Source/Text3D/Private/Text3DComponent.cpp b/Engine/Plugins/Experimental/Text3D/Source/Text3D/Private/Text3DComponent.cpp
--- a/Engine/Plugins/Experimental/Text3D/Source/Text3D/Private/Text3DComponent.cpp
+++ b/Engine/Plugins/Experimental/Text3D/Source/Text3D/Private/Text3DComponent.cpp
@@ -46,59 +46,79 @@ struct FText3DShapedText
 	void CalculateWidth()
 	{
 		TArray<FShapedGlyphLine> NewLines;
+		NewLines.Reserve(Lines.Num());
+		
 		for (const FShapedGlyphLine& GlyphLine : Lines)
 		{
-			FShapedGlyphLine& CurrentLine = NewLines.Add_GetRef(FShapedGlyphLine());
+			FShapedGlyphLine* CurrentLine = &NewLines.AddDefaulted_GetRef();
 			TArray<FShapedGlyphEntry> CurrentWord;
 
 			float LineWidth = 0.0f;
-			int32 GlyphCount = GlyphLine.GlyphsToRender.Num();
+			const int32 GlyphCount = GlyphLine.GlyphsToRender.Num();
 			float CurrentWordLength = 0.0f;
+			
 			for (int32 GlyphIdx = 0; GlyphIdx < GlyphCount; ++GlyphIdx)
 			{
 				const FShapedGlyphEntry& CurrentGlyph = GlyphLine.GlyphsToRender[GlyphIdx];
-				float GlyphAdv = GlyphLine.GetAdvance(GlyphIdx, Kerning, WordSpacing);
+
+				// Trim for proper positioning
+				if (!CurrentGlyph.bIsVisible && GlyphIdx == 0)
+				{
+					continue;
+				}
+
+				const bool bWordBreak = !CurrentGlyph.bIsVisible || GlyphIdx == GlyphLine.GlyphsToRender.Num() - 1;
+				const float GlyphAdv = GlyphLine.GetAdvance(GlyphIdx, Kerning, WordSpacing);
+
+				// If we're at the end the line or at whitespace
+				if (bWrap                               // when we're wrapping
+					&& bWordBreak                       // and at a word break
+					&& LineWidth > MaxWidth             // and the current line is longer than the max
+					&& CurrentWordLength != LineWidth)  // and the line is not just a single word that we can't break
+				{
+					CurrentLine->Width = LineWidth - CurrentWordLength;
+					CurrentLine = &NewLines.AddDefaulted_GetRef();
+					LineWidth = CurrentWordLength;
+				}
+
+				CurrentWord.Add(CurrentGlyph);
 				LineWidth += GlyphAdv;
 				CurrentWordLength += GlyphAdv;
 
-				// If we're at the end the line or at whitespace
-				if (!CurrentGlyph.bIsVisible || GlyphIdx == GlyphLine.GlyphsToRender.Num() - 1)
+				if (bWordBreak)
 				{
-					if (bWrap && LineWidth > MaxWidth)
-					{
-						CurrentLine.Width = LineWidth - CurrentWordLength;
-
-						if (CurrentWordLength != LineWidth) // No break to wrap
-						{
-							NewLines.Add(FShapedGlyphLine());
-							LineWidth = CurrentWordLength;
-						}
-					}
-					else
-					{
-						CurrentLine.Width = LineWidth;
-					}
-
-					// if we're wrapping, we may/may not want the white space
-					if (!bWrap || LineWidth < MaxWidth || LineWidth == CurrentWordLength)
-					{
-						CurrentWord.Add(CurrentGlyph);
-					}
-
-					NewLines.Last().GlyphsToRender.Append(CurrentWord);
+					CurrentLine->GlyphsToRender.Append(CurrentWord);
 					CurrentWordLength = 0.0f;
 					CurrentWord.Empty();
 				}
-				else
-				{
-					CurrentWord.Add(CurrentGlyph);
-				}
 			}
 
-			CurrentLine.GlyphsToRender.Append(CurrentWord);
-			CurrentLine.Width = LineWidth;
+			CurrentLine->Width = LineWidth;
 		}
+		
+		for (FShapedGlyphLine& NewLine : NewLines)
+		{
+			if (NewLine.GlyphsToRender.IsEmpty())
+			{
+				continue;
+			}
 
+			const FShapedGlyphEntry& FirstGlyph = NewLine.GlyphsToRender[0];
+			if (!FirstGlyph.bIsVisible)
+			{
+				NewLine.Width -= NewLine.GetAdvance(0, Kerning, WordSpacing);
+				NewLine.GlyphsToRender.RemoveAt(0);
+			}
+
+			const int32 LastIndex = NewLine.GlyphsToRender.Num() - 1;
+			const FShapedGlyphEntry& LastGlyph = NewLine.GlyphsToRender[LastIndex];
+			if (!LastGlyph.bIsVisible)
+			{
+				NewLine.Width -= NewLine.GetAdvance(LastIndex, Kerning, WordSpacing);
+				NewLine.GlyphsToRender.RemoveAt(LastIndex);
+			}
+		}
+		
 		Lines = NewLines;
 	}
 
@@ -1130,7 +1150,7 @@ void UText3DComponent::UpdateTransforms()
 	ShapedText->Kerning = Kerning;
 	ShapedText->WordSpacing = WordSpacing;
 	ShapedText->MaxWidth = MaxWidth;
-	ShapedText->bWrap = MaxWidthHandling == EText3DMaxWidthHandling::WrapAndScale;
+	ShapedText->bWrap = bHasMaxWidth && MaxWidthHandling == EText3DMaxWidthHandling::WrapAndScale;
 
 	ShapedText->CalculateWidth();
 	CalculateTextScale();
@@ -1288,7 +1308,7 @@ void UText3DComponent::BuildTextMeshInternal(const bool& bCleanCache)
 	ShapedText->Kerning = Kerning;
 	ShapedText->WordSpacing = WordSpacing;
 	ShapedText->MaxWidth = MaxWidth;
-	ShapedText->bWrap = MaxWidthHandling == EText3DMaxWidthHandling::WrapAndScale;
+	ShapedText->bWrap = bHasMaxWidth && MaxWidthHandling == EText3DMaxWidthHandling::WrapAndScale;
 
 	constexpr int32 AdjustedFontSize = 48; // Magic number that makes font scale consistent with previous implementation
 	FSlateFontInfo FontInfo(Font, AdjustedFontSize);
