MotionDesign : Text3D
- Fix text duplication issue causing glyphs to disappear or render incorrectly

diff --git a/Engine/Plugins/Experimental/Text3D/Source/Text3D/Private/Text3DComponent.cpp b/Engine/Plugins/Experimental/Text3D/Source/Text3D/Private/Text3DComponent.cpp
--- a/Engine/Plugins/Experimental/Text3D/Source/Text3D/Private/Text3DComponent.cpp
+++ b/Engine/Plugins/Experimental/Text3D/Source/Text3D/Private/Text3DComponent.cpp
@@ -355,6 +355,20 @@ void UText3DComponent::OnUnregister()
 	Super::OnUnregister();
 }
 
+void UText3DComponent::PostEditImport()
+{
+	Super::PostEditImport();
+
+	ModifyFlags = EText3DModifyFlags::All;
+}
+
+void UText3DComponent::PostDuplicate(EDuplicateMode::Type InDuplicateMode)
+{
+	Super::PostDuplicate(InDuplicateMode);
+
+	ModifyFlags = EText3DModifyFlags::All;
+}
+
 #if WITH_EDITOR
 void UText3DComponent::PostEditChangeProperty(FPropertyChangedEvent& PropertyChangedEvent)
 {
diff --git a/Engine/Plugins/Experimental/Text3D/Source/Text3D/Public/Text3DComponent.h b/Engine/Plugins/Experimental/Text3D/Source/Text3D/Public/Text3DComponent.h
--- a/Engine/Plugins/Experimental/Text3D/Source/Text3D/Public/Text3DComponent.h
+++ b/Engine/Plugins/Experimental/Text3D/Source/Text3D/Public/Text3DComponent.h
@@ -76,15 +76,21 @@ class TEXT3D_API UText3DComponent : public USceneComponent
 public:
 	UText3DComponent();
 
-	virtual void BeginDestroy() override;
-	virtual void PostLoad() override;
+	//~ Begin UActorComponent
 	virtual void OnRegister() override;
 	virtual void OnUnregister() override;
-
+	//~ End UActorComponent
+	
+	//~ Begin UObject
+	virtual void BeginDestroy() override;
+	virtual void PostLoad() override;
+	virtual void PostEditImport() override;
+	virtual void PostDuplicate(EDuplicateMode::Type InDuplicateMode) override;
 #if WITH_EDITOR
-	virtual void PostEditChangeProperty(struct FPropertyChangedEvent& PropertyChangedEvent) override;
+	virtual void PostEditChangeProperty(FPropertyChangedEvent& PropertyChangedEvent) override;
 	virtual void PostTransacted(const FTransactionObjectEvent& TransactionEvent) override;
 #endif
+	//~ End UObject
 
 	/**
 	 * Delegate called after text is rebuilt
@@ -95,7 +101,10 @@ public:
 	 * Delegate called after text is rebuilt
 	 */
 	DECLARE_MULTICAST_DELEGATE(FTextGeneratedNative);
-	FTextGeneratedNative& OnTextGenerated() { return TextGeneratedNativeDelegate; }
+	FTextGeneratedNative::RegistrationType& OnTextGenerated()
+	{
+		return TextGeneratedNativeDelegate;
+	}
 
 	/** Get whether to allow automatic refresh/mesh generation */
 	bool RefreshesOnChange() const;
