MotionDesign : Text3D
- Add a base text renderer support and implementing legacy static meshes renderer
- Updating files and folders structure for better maintainability
- Renaming files and splitting code
- Fix cast shadow option not working

diff --git a/Engine/Plugins/Experimental/Avalanche/Source/AvalancheMask/Private/Handling/AvaMaskText3DActorMaterialCollectionHandle.cpp b/Engine/Plugins/Experimental/Avalanche/Source/AvalancheMask/Private/Handling/AvaMaskText3DActorMaterialCollectionHandle.cpp
--- a/Engine/Plugins/Experimental/Avalanche/Source/AvalancheMask/Private/Handling/AvaMaskText3DActorMaterialCollectionHandle.cpp
+++ b/Engine/Plugins/Experimental/Avalanche/Source/AvalancheMask/Private/Handling/AvaMaskText3DActorMaterialCollectionHandle.cpp
@@ -14,7 +14,6 @@
 #include "IAvaMaskMaterialHandle.h"
 #include "MaterialHub/AvaTextMaterialHub.h"
 #include "Materials/MaterialInterface.h"
-#include "Mesh.h"
 #include "Text3DComponent.h"
 
 namespace UE::AvaMask::Private
diff --git a/Engine/Plugins/Experimental/Avalanche/Source/AvalancheText/Public/AvaText3DComponent.h b/Engine/Plugins/Experimental/Avalanche/Source/AvalancheText/Public/AvaText3DComponent.h
--- a/Engine/Plugins/Experimental/Avalanche/Source/AvalancheText/Public/AvaText3DComponent.h
+++ b/Engine/Plugins/Experimental/Avalanche/Source/AvalancheText/Public/AvaText3DComponent.h
@@ -324,7 +324,7 @@ private:
 	EAvaTextTranslucency TranslucencyStyle = EAvaTextTranslucency::None;
 
 	UPROPERTY(EditAnywhere, Getter="GetOpacity", Setter="SetOpacity", Category="Style", meta=(ClampMin=0.0, ClampMax=1.0, EditCondition="TranslucencyStyle == EAvaTextTranslucency::Translucent && ColoringStyle != EAvaTextColoringStyle::CustomMaterial", EditConditionHides, DisplayAfter="Text", AllowPrivateAccess="true"))
-	float Opacity = 0.75f;
+	float Opacity = 1.f;
 
 	UPROPERTY(EditAnywhere, Getter="GetMaskOrientation", Setter="SetMaskOrientation", Category="Style", meta=(EditCondition="TranslucencyStyle == EAvaTextTranslucency::GradientMask", EditConditionHides, DisplayAfter="Text", AllowPrivateAccess="true"))
 	EAvaMaterialMaskOrientation MaskOrientation = EAvaMaterialMaskOrientation::LeftRight;
diff --git a/Engine/Plugins/Experimental/Avalanche/Source/AvalancheText/Public/AvaTextCharacterTransform.h b/Engine/Plugins/Experimental/Avalanche/Source/AvalancheText/Public/AvaTextCharacterTransform.h
--- a/Engine/Plugins/Experimental/Avalanche/Source/AvalancheText/Public/AvaTextCharacterTransform.h
+++ b/Engine/Plugins/Experimental/Avalanche/Source/AvalancheText/Public/AvaTextCharacterTransform.h
@@ -2,7 +2,7 @@
 
 #pragma once
 
-#include "Text3DCharacterTransform.h"
+#include "Characters/Text3DCharacterTransform.h"
 
 #include "AvaTextCharacterTransform.generated.h"
 
diff --git a/Engine/Plugins/Experimental/Text3D/Source/Text3D/Private/Arrangement.cpp b/Engine/Plugins/Experimental/Text3D/Source/Text3D/Private/Arrangement.cpp
deleted file mode 100644
diff --git a/Engine/Plugins/Experimental/Text3D/Source/Text3D/Private/Arrangement.h b/Engine/Plugins/Experimental/Text3D/Source/Text3D/Private/Arrangement.h
deleted file mode 100644
diff --git a/Engine/Plugins/Experimental/Text3D/Source/Text3D/Private/Characters/Text3DCharacterTransform.cpp b/Engine/Plugins/Experimental/Text3D/Source/Text3D/Private/Characters/Text3DCharacterTransform.cpp
new file mode 100644
--- /dev/null
+++ b/Engine/Plugins/Experimental/Text3D/Source/Text3D/Private/Characters/Text3DCharacterTransform.cpp
@@ -0,0 +1,430 @@
+// Copyright Epic Games, Inc. All Rights Reserved.
+
+#include "Characters/Text3DCharacterTransform.h"
+#include "GameFramework/Actor.h"
+#include "Components/StaticMeshComponent.h"
+#include "Renderers/StaticMeshes/Text3DStaticMeshesRenderer.h"
+#include "Text3DComponent.h"
+
+UText3DCharacterTransform::UText3DCharacterTransform()
+{
+	bLocationEnabled = false;
+	LocationProgress = 0.0f;
+	LocationOrder = EText3DCharacterEffectOrder::Normal;
+	LocationRange = 50.0f;
+	LocationDistance = FVector(100.0f, 0.0f, 0.0f);
+
+	bScaleEnabled = false;
+	ScaleProgress = 0.0f;
+	ScaleOrder = EText3DCharacterEffectOrder::Normal;
+	ScaleRange = 50.0f;
+	ScaleBegin = FVector(1.0f, 0.0f, 0.0f);
+	ScaleEnd = FVector(1.0f);
+
+	bRotateEnabled = false;
+	RotateProgress = 0.0f;
+	RotateOrder = EText3DCharacterEffectOrder::Normal;
+	RotateRange = 50.0f;
+	RotateBegin = FRotator(-90.0f, 0.0f, 0.0f);
+	RotateEnd = FRotator(0.0f, 0.0f, 0.0f);
+
+	bInitialized = false;
+}
+
+void UText3DCharacterTransform::OnRegister()
+{
+	Super::OnRegister();
+
+	if (!bInitialized)
+	{
+		bInitialized = true;
+		ProcessEffect();
+	}
+
+	if (UText3DComponent* Text3DComponent = GetText3DComponent())
+	{
+		Text3DComponent->OnTextGenerated().AddUObject(this, &UText3DCharacterTransform::ProcessEffect);
+	}
+}
+
+void UText3DCharacterTransform::OnUnregister()
+{
+	if (UText3DComponent* Text3DComponent = GetText3DComponent())
+	{
+		Text3DComponent->OnTextGenerated().RemoveAll(this);
+	}
+
+	Super::OnUnregister();
+}
+
+#if WITH_EDITOR
+void UText3DCharacterTransform::PostEditChangeProperty(FPropertyChangedEvent& PropertyChangedEvent)
+{
+	Super::PostEditChangeProperty(PropertyChangedEvent);
+
+	const FName Name = PropertyChangedEvent.GetPropertyName();
+	if (Name == GET_MEMBER_NAME_CHECKED(UText3DCharacterTransform, bLocationEnabled))
+	{
+		ResetLocation();
+	}
+	else if (Name == GET_MEMBER_NAME_CHECKED(UText3DCharacterTransform, bScaleEnabled))
+	{
+		ResetScale();
+	}
+	else if (Name == GET_MEMBER_NAME_CHECKED(UText3DCharacterTransform, bRotateEnabled))
+	{
+		ResetRotate();
+	}
+
+	ProcessEffect();
+}
+#endif
+
+void UText3DCharacterTransform::ProcessEffect()
+{
+	UText3DComponent* Text3DComponent = GetText3DComponent();
+	if (!Text3DComponent)
+	{
+		return;
+	}
+
+	// This will change to communicate directly with per CharacterData and renderer will not be used at all
+	UText3DStaticMeshesRenderer* TextRenderer = Cast<UText3DStaticMeshesRenderer>(Text3DComponent->GetTextRenderer());
+
+	int32 Count = TextRenderer->GetGlyphCount();
+	for (int32 Index = 0; Index < Count; Index++)
+	{
+		USceneComponent* GlyphComponent = TextRenderer->GetGlyphMeshComponent(Index);
+		if (!GlyphComponent)
+		{
+			continue;
+		}
+
+		if (bLocationEnabled)
+		{
+			float Effect = GetEffectValue(Index, Count, LocationOrder, LocationProgress, LocationRange);
+			FVector Location = LocationDistance * Effect;
+			GlyphComponent->SetRelativeLocation(Location);
+		}
+
+		if (bScaleEnabled)
+		{
+			float Effect = GetEffectValue(Index, Count, ScaleOrder, ScaleProgress, ScaleRange);
+			FVector Scale = ScaleBegin + ((ScaleEnd - ScaleBegin) * Effect);
+			GlyphComponent->SetRelativeScale3D(Scale);
+		}
+
+		if (bRotateEnabled)
+		{
+			float Effect = GetEffectValue(Index, Count, RotateOrder, RotateProgress, RotateRange);
+			FRotator Rotator = RotateBegin + (RotateEnd - RotateBegin) * Effect;
+			GlyphComponent->SetRelativeRotation(Rotator);
+		}
+	}
+}
+
+void UText3DCharacterTransform::ResetLocation()
+{
+	UText3DComponent* Text3DComponent = GetText3DComponent();
+	if (!Text3DComponent)
+	{
+		return;
+	}
+
+	// This will change to communicate directly with per CharacterData and renderer will not be used at all
+	UText3DStaticMeshesRenderer* TextRenderer = Cast<UText3DStaticMeshesRenderer>(Text3DComponent->GetTextRenderer());
+
+	int32 Count = TextRenderer->GetGlyphCount();
+	for (int32 Index = 0; Index < Count; Index++)
+	{
+		USceneComponent* GlyphComponent = TextRenderer->GetGlyphMeshComponent(Index);
+		if (GlyphComponent)
+		{
+			GlyphComponent->SetRelativeLocation(FVector::ZeroVector);
+		}
+	}
+}
+
+void UText3DCharacterTransform::ResetRotate()
+{
+	UText3DComponent* Text3DComponent = GetText3DComponent();
+	if (!Text3DComponent)
+	{
+		return;
+	}
+
+	// This will change to communicate directly with per CharacterData and renderer will not be used at all
+	UText3DStaticMeshesRenderer* TextRenderer = Cast<UText3DStaticMeshesRenderer>(Text3DComponent->GetTextRenderer());
+
+	int32 Count = TextRenderer->GetGlyphCount();
+	for (int32 Index = 0; Index < Count; Index++)
+	{
+		USceneComponent* GlyphComponent = TextRenderer->GetGlyphMeshComponent(Index);
+		if (GlyphComponent)
+		{
+			GlyphComponent->SetRelativeRotation(FRotator::ZeroRotator);
+		}
+	}
+}
+
+void UText3DCharacterTransform::ResetScale()
+{
+	UText3DComponent* Text3DComponent = GetText3DComponent();
+	if (!Text3DComponent)
+	{
+		return;
+	}
+
+	// This will change to communicate directly with per CharacterData and renderer will not be used at all
+	UText3DStaticMeshesRenderer* TextRenderer = Cast<UText3DStaticMeshesRenderer>(Text3DComponent->GetTextRenderer());
+
+	int32 Count = TextRenderer->GetGlyphCount();
+	for (int32 Index = 0; Index < Count; Index++)
+	{
+		USceneComponent* GlyphComponent = TextRenderer->GetGlyphMeshComponent(Index);
+		if (GlyphComponent)
+		{
+			GlyphComponent->SetRelativeScale3D(FVector(1.0f));
+		}
+	}
+}
+
+float UText3DCharacterTransform::GetEffectValue(int32 Index, int32 Total, EText3DCharacterEffectOrder Order, float Progress, float Range)
+{
+	float Effect = 0.0f, Strip = 1.0f;
+	GetLineParameters(Range * 0.01f, Order, Total, Effect, Strip);
+
+	int32 Position = GetEffectPosition(Index, Total, Order);
+	return FMath::Clamp((Progress * 0.01f - Effect * Position) / Strip, 0.0f, 1.0f);
+}
+
+void UText3DCharacterTransform::GetLineParameters(float Range, EText3DCharacterEffectOrder Order, int32 Total, float & EffectOut, float & StripOut)
+{
+	float Segment = 1.0f / Total;
+	EffectOut = (1.0f - Range) / Total;
+	StripOut = FMath::Max(Range, Segment);
+
+	if (Order == EText3DCharacterEffectOrder::FromCenter || Order == EText3DCharacterEffectOrder::ToCenter)
+	{
+		int32 Center = Total / 2.0f - 0.5f;
+		EffectOut *= (float)Total / (Center + 1.0f);
+	}
+
+	if (FMath::IsNearlyZero(EffectOut))
+	{
+		EffectOut = 0.01f;
+	}
+}
+
+int32 UText3DCharacterTransform::GetEffectPosition(int32 Index, int32 Total, EText3DCharacterEffectOrder Order)
+{
+	int32 Center = 0.5f * Total - 0.5f;
+	switch (Order)
+	{
+	case EText3DCharacterEffectOrder::FromCenter:
+	{
+		if (Index > Center)
+		{
+			Index = Total - Index - 1;
+		}
+
+		return Center - Index;
+	}
+
+	case EText3DCharacterEffectOrder::ToCenter:
+	{
+		if (Index > Center)
+		{
+			return Total - Index - 1;
+		}
+
+		return Index;
+	}
+
+	case EText3DCharacterEffectOrder::Opposite:
+	{
+		return Total - Index - 1;
+	}
+	}
+
+	return Index;
+}
+
+void UText3DCharacterTransform::SetLocationEnabled(bool bEnabled)
+{
+	if (bLocationEnabled != bEnabled)
+	{
+		bLocationEnabled = bEnabled;
+		ResetLocation();
+		ProcessEffect();
+	}
+}
+
+void UText3DCharacterTransform::SetLocationProgress(float Progress)
+{
+	Progress = FMath::Clamp(Progress, 0.0f, 100.0f);
+	if (!FMath::IsNearlyEqual(LocationProgress, Progress))
+	{
+		LocationProgress = Progress;
+		ProcessEffect();
+	}
+}
+
+void UText3DCharacterTransform::SetLocationOrder(EText3DCharacterEffectOrder Order)
+{
+	if (LocationOrder != Order)
+	{
+		LocationOrder = Order;
+		ProcessEffect();
+	}
+}
+
+void UText3DCharacterTransform::SetLocationRange(float Range)
+{
+	Range = FMath::Clamp(Range, 0.0f, 100.0f);
+	if (!FMath::IsNearlyEqual(LocationRange, Range))
+	{
+		LocationRange = Range;
+		ProcessEffect();
+	}
+}
+
+void UText3DCharacterTransform::SetLocationDistance(FVector Distance)
+{
+	if (!(LocationDistance - Distance).IsNearlyZero())
+	{
+		LocationDistance = Distance;
+		ProcessEffect();
+	}
+}
+
+void UText3DCharacterTransform::SetScaleEnabled(bool bEnabled)
+{
+	if (bScaleEnabled != bEnabled)
+	{
+		bScaleEnabled = bEnabled;
+		ResetScale();
+		ProcessEffect();
+	}
+}
+
+void UText3DCharacterTransform::SetScaleProgress(float Progress)
+{
+	Progress = FMath::Clamp(Progress, 0.0f, 100.0f);
+	if (!FMath::IsNearlyEqual(ScaleProgress, Progress))
+	{
+		ScaleProgress = Progress;
+		ProcessEffect();
+	}
+}
+
+void UText3DCharacterTransform::SetScaleOrder(EText3DCharacterEffectOrder Order)
+{
+	if (ScaleOrder != Order)
+	{
+		ScaleOrder = Order;
+		ProcessEffect();
+	}
+}
+
+void UText3DCharacterTransform::SetScaleRange(float Range)
+{
+	Range = FMath::Clamp(Range, 0.0f, 100.0f);
+	if (!FMath::IsNearlyEqual(ScaleRange, Range))
+	{
+		ScaleRange = Range;
+		ProcessEffect();
+	}
+}
+
+void UText3DCharacterTransform::SetScaleBegin(FVector Value)
+{
+	if (!(ScaleBegin - Value).IsNearlyZero())
+	{
+		ScaleBegin = Value;
+		ProcessEffect();
+	}
+}
+
+void UText3DCharacterTransform::SetScaleEnd(FVector Value)
+{
+	if (!(ScaleEnd - Value).IsNearlyZero())
+	{
+		ScaleEnd = Value;
+		ProcessEffect();
+	}
+}
+
+void UText3DCharacterTransform::SetRotateEnabled(bool bEnabled)
+{
+	if (bRotateEnabled != bEnabled)
+	{
+		bRotateEnabled = bEnabled;
+		ResetRotate();
+		ProcessEffect();
+	}
+}
+
+void UText3DCharacterTransform::SetRotateProgress(float Progress)
+{
+	Progress = FMath::Clamp(Progress, 0.0f, 100.0f);
+	if (!FMath::IsNearlyEqual(RotateProgress, Progress))
+	{
+		RotateProgress = Progress;
+		ProcessEffect();
+	}
+}
+
+void UText3DCharacterTransform::SetRotateOrder(EText3DCharacterEffectOrder Order)
+{
+	if (RotateOrder != Order)
+	{
+		RotateOrder = Order;
+		ProcessEffect();
+	}
+}
+
+void UText3DCharacterTransform::SetRotateRange(float Range)
+{
+	Range = FMath::Clamp(Range, 0.0f, 100.0f);
+	if (!FMath::IsNearlyEqual(RotateRange, Range))
+	{
+		RotateRange = Range;
+		ProcessEffect();
+	}
+}
+
+void UText3DCharacterTransform::SetRotateBegin(FRotator Value)
+{
+	if (!(RotateBegin - Value).IsNearlyZero())
+	{
+		RotateBegin = Value;
+		ProcessEffect();
+	}
+}
+
+void UText3DCharacterTransform::SetRotateEnd(FRotator Value)
+{
+	if (!(RotateEnd - Value).IsNearlyZero())
+	{
+		RotateEnd = Value;
+		ProcessEffect();
+	}
+}
+
+UText3DComponent* UText3DCharacterTransform::GetText3DComponent()
+{
+	UText3DComponent* Component = Cast<UText3DComponent>(GetAttachParent());
+	if (Component)
+	{
+		return Component;
+	}
+
+	AActor* Owner = GetOwner();
+	if (!Owner)
+	{
+		return nullptr;
+	}
+	
+	return Owner->FindComponentByClass<UText3DComponent>();
+}
diff --git a/Engine/Plugins/Experimental/Text3D/Source/Text3D/Private/Contour.cpp b/Engine/Plugins/Experimental/Text3D/Source/Text3D/Private/Contour.cpp
deleted file mode 100644
diff --git a/Engine/Plugins/Experimental/Text3D/Source/Text3D/Private/Contour.h b/Engine/Plugins/Experimental/Text3D/Source/Text3D/Private/Contour.h
deleted file mode 100644
diff --git a/Engine/Plugins/Experimental/Text3D/Source/Text3D/Private/ContourList.cpp b/Engine/Plugins/Experimental/Text3D/Source/Text3D/Private/ContourList.cpp
deleted file mode 100644
diff --git a/Engine/Plugins/Experimental/Text3D/Source/Text3D/Private/ContourList.h b/Engine/Plugins/Experimental/Text3D/Source/Text3D/Private/ContourList.h
deleted file mode 100644
diff --git a/Engine/Plugins/Experimental/Text3D/Source/Text3D/Private/ContourNode.h b/Engine/Plugins/Experimental/Text3D/Source/Text3D/Private/ContourNode.h
deleted file mode 100644
diff --git a/Engine/Plugins/Experimental/Text3D/Source/Text3D/Private/Data.cpp b/Engine/Plugins/Experimental/Text3D/Source/Text3D/Private/Data.cpp
deleted file mode 100644
diff --git a/Engine/Plugins/Experimental/Text3D/Source/Text3D/Private/Data.h b/Engine/Plugins/Experimental/Text3D/Source/Text3D/Private/Data.h
deleted file mode 100644
diff --git a/Engine/Plugins/Experimental/Text3D/Source/Text3D/Private/GeometryBuilders/Text3DGlyph.cpp b/Engine/Plugins/Experimental/Text3D/Source/Text3D/Private/GeometryBuilders/Text3DGlyph.cpp
new file mode 100644
--- /dev/null
+++ b/Engine/Plugins/Experimental/Text3D/Source/Text3D/Private/GeometryBuilders/Text3DGlyph.cpp
@@ -0,0 +1,77 @@
+// Copyright Epic Games, Inc. All Rights Reserved.
+
+#include "GeometryBuilders/Text3DGlyph.h"
+
+#include "Engine/StaticMesh.h"
+#include "Materials/Material.h"
+#include "Text3DInternalTypes.h"
+
+FText3DGlyph::FText3DGlyph() :
+	StaticMeshAttributes(MeshDescription)
+{
+	StaticMeshAttributes.Register();
+	Groups.SetNum(static_cast<int32>(EText3DGroupType::TypeCount));
+	MeshDescription.ReserveNewPolygonGroups(Groups.Num());
+
+	for (int32 Index = 0; Index < Groups.Num(); Index++)
+	{
+		MeshDescription.CreatePolygonGroup();
+	}
+}
+
+void FText3DGlyph::Build(UStaticMesh* StaticMesh, UMaterial* DefaultMaterial)
+{
+	check(StaticMesh);
+
+	auto AddMaterial = [](UStaticMesh* InMesh, UMaterial* InMaterial, int32 InIndex)
+	{
+		using namespace UE::Text3D::Materials;
+
+		check(SlotNames.IsValidIndex(InIndex))
+
+		FName MaterialName = SlotNames[InIndex];
+
+#if WITH_EDITORONLY_DATA
+		FStaticMaterial& StaticMaterial = InMesh->GetStaticMaterials().Emplace_GetRef(InMaterial, MaterialName, MaterialName);
+#else
+		FStaticMaterial& StaticMaterial = InMesh->GetStaticMaterials().Emplace_GetRef(InMaterial, MaterialName);
+#endif
+
+		StaticMaterial.UVChannelData = FMeshUVChannelInfo(1.0f);
+
+		return MaterialName;
+	};
+
+	for (int32 Index = 0; Index < Groups.Num(); Index++)
+	{
+		const FPolygonGroupID PolyGroup(Index);
+		if (MeshDescription.GetNumPolygonGroupTriangles(PolyGroup) != 0)
+		{
+			StaticMeshAttributes.GetPolygonGroupMaterialSlotNames()[PolyGroup] = AddMaterial(StaticMesh, DefaultMaterial, Index);
+		}
+	}
+
+	TArray<const FMeshDescription*> MeshDescriptions;
+	MeshDescriptions.Add(&MeshDescription);
+
+	UStaticMesh::FBuildMeshDescriptionsParams Params;
+	Params.bCommitMeshDescription = true;
+	Params.bFastBuild = true;
+	Params.bAllowCpuAccess = true;
+	StaticMesh->BuildFromMeshDescriptions(MeshDescriptions, Params);
+}
+
+FMeshDescription& FText3DGlyph::GetMeshDescription()
+{
+	return MeshDescription;
+}
+
+FStaticMeshAttributes& FText3DGlyph::GetStaticMeshAttributes()
+{
+	return StaticMeshAttributes;
+}
+
+TText3DGroupList& FText3DGlyph::GetGroups()
+{
+	return Groups;
+}
diff --git a/Engine/Plugins/Experimental/Text3D/Source/Text3D/Private/GeometryBuilders/Text3DGlyph.h b/Engine/Plugins/Experimental/Text3D/Source/Text3D/Private/GeometryBuilders/Text3DGlyph.h
new file mode 100644
--- /dev/null
+++ b/Engine/Plugins/Experimental/Text3D/Source/Text3D/Private/GeometryBuilders/Text3DGlyph.h
@@ -0,0 +1,36 @@
+// Copyright Epic Games, Inc. All Rights Reserved.
+
+#pragma once
+
+#include "Containers/Array.h"
+#include "MeshDescription.h"
+#include "StaticMeshAttributes.h"
+#include "Text3DTypes.h"
+
+class UMaterial;
+class UStaticMesh;
+
+struct FText3DPolygonGroup
+{
+	int32 FirstVertex;
+	int32 FirstTriangle;
+};
+
+using TText3DGroupList = TArray<FText3DPolygonGroup, TFixedAllocator<static_cast<int32>(EText3DGroupType::TypeCount)>>;
+
+class FText3DGlyph
+{
+public:
+	FText3DGlyph();
+
+	void Build(UStaticMesh* StaticMesh, UMaterial* DefaultMaterial);
+
+	FMeshDescription& GetMeshDescription();
+	FStaticMeshAttributes& GetStaticMeshAttributes();
+	TText3DGroupList& GetGroups();
+
+private:
+	FMeshDescription MeshDescription;
+	FStaticMeshAttributes StaticMeshAttributes;
+	TText3DGroupList Groups;
+};
diff --git a/Engine/Plugins/Experimental/Text3D/Source/Text3D/Private/GeometryBuilders/Text3DGlyphArrangement.cpp b/Engine/Plugins/Experimental/Text3D/Source/Text3D/Private/GeometryBuilders/Text3DGlyphArrangement.cpp
new file mode 100644
--- /dev/null
+++ b/Engine/Plugins/Experimental/Text3D/Source/Text3D/Private/GeometryBuilders/Text3DGlyphArrangement.cpp
@@ -0,0 +1,316 @@
+// Copyright Epic Games, Inc. All Rights Reserved.
+
+#include "GeometryBuilders/Text3DGlyphArrangement.h"
+
+using namespace UE::Geometry;
+
+FText3DGlyphArrangement::FText3DGlyphArrangement(const FAxisAlignedBox2f& InBoundsHint)
+    : PointHash(static_cast<double>(InBoundsHint.MaxDim()) / 64, -1)
+{
+}
+
+void FText3DGlyphArrangement::Insert(const FSegment2f& Segment)
+{
+    const FVector2f A = Segment.StartPoint();
+    const FVector2f B = Segment.EndPoint();
+
+    InsertSegment({static_cast<double>(A.X), static_cast<double>(A.Y)}, {static_cast<double>(B.X), static_cast<double>(B.Y)}, VertexSnapTol);
+}
+
+int FText3DGlyphArrangement::InsertPoint(const FVector2d &InPoint, double InTol)
+{
+    int PIdx = FindExistingVertex(InPoint);
+    if (PIdx > -1)
+    {
+        return -1;
+    }
+
+    // TODO: currently this tries to add the vertex on the closest edge below tolerance; we should instead insert at *every* edge below tolerance!  ... but that is more inconvenient to write
+    FVector2d x = FVector2d::Zero(), y = FVector2d::Zero();
+    double ClosestDistSq = InTol*InTol;
+    int FoundEdgeToSplit = -1;
+    for (int EID = 0, ExistingEdgeMax = Graph.MaxEdgeID(); EID < ExistingEdgeMax; EID++)
+    {
+        if (!Graph.IsEdge(EID))
+        {
+            continue;
+        }
+
+        Graph.GetEdgeV(EID, x, y);
+        FSegment2d Seg(x, y);
+        double DistSq = Seg.DistanceSquared(InPoint);
+        if (DistSq < ClosestDistSq)
+        {
+            ClosestDistSq = DistSq;
+            FoundEdgeToSplit = EID;
+        }
+    }
+    if (FoundEdgeToSplit > -1)
+    {
+        FDynamicGraph2d::FEdgeSplitInfo splitInfo;
+        EMeshResult result = Graph.SplitEdge(FoundEdgeToSplit, splitInfo);
+        ensureMsgf(result == EMeshResult::Ok, TEXT("insert_into_segment: edge split failed?"));
+        Graph.SetVertex(splitInfo.VNew, InPoint);
+        PointHash.InsertPointUnsafe(splitInfo.VNew, InPoint);
+        return splitInfo.VNew;
+    }
+
+    int VID = Graph.AppendVertex(InPoint);
+    PointHash.InsertPointUnsafe(VID, InPoint);
+    return VID;
+}
+
+bool FText3DGlyphArrangement::InsertSegment(FVector2d InA, FVector2d InB, double InTol)
+{
+    // handle degenerate edges
+    int a_idx = FindExistingVertex(InA);
+    int b_idx = FindExistingVertex(InB);
+    if (a_idx == b_idx && a_idx >= 0)
+    {
+        return false;
+    }
+    // snap input vertices
+    if (a_idx >= 0)
+    {
+        InA = Graph.GetVertex(a_idx);
+    }
+    if (b_idx >= 0)
+    {
+        InB = Graph.GetVertex(b_idx);
+    }
+
+    // handle tiny-segment case
+    double SegLenSq = DistanceSquared(InA, InB);
+    if (SegLenSq <= VertexSnapTol*VertexSnapTol)
+    {
+        // seg is too short and was already on an existing vertex; just consider that vertex to be the inserted segment
+        if (a_idx >= 0 || b_idx >= 0)
+        {
+            return false;
+        }
+        // seg is too short and wasn't on an existing vertex; add it as an isolated vertex
+        return InsertPoint(InA, InTol) != -1;
+    }
+
+    // ok find all intersections
+    TArray<FIntersection> Hits;
+    FindIntersectingEdges(InA, InB, Hits, InTol);
+
+    // we are going to construct a list of <T,vertex_id> values along segment AB
+    TArray<FSegmentPoint> points;
+    FSegment2d segAB = FSegment2d(InA, InB);
+
+    FindIntersectingVertices(segAB, a_idx, b_idx, points, InTol);
+
+    // insert intersections into existing segments
+    for (int i = 0, N = Hits.Num(); i < N; ++i)
+    {
+        FIntersection Intr = Hits[i];
+        int EID = Intr.EID;
+        double t0 = Intr.Intr.Parameter0, t1 = Intr.Intr.Parameter1;
+
+        // insert first point at t0
+        int new_eid = -1;
+        if (Intr.Intr.Type == EIntersectionType::Point || Intr.Intr.Type == EIntersectionType::Segment)
+        {
+            FIndex2i new_info = SplitSegmentAtDistance(EID, t0, VertexSnapTol);
+            new_eid = new_info.B;
+            FVector2d v = Graph.GetVertex(new_info.A);
+            points.Add(FSegmentPoint{segAB.Project(v), new_info.A});
+        }
+
+        // if intersection was on-segment, then we have a second point at t1
+        if (Intr.Intr.Type == EIntersectionType::Segment)
+        {
+            if (new_eid == -1)
+            {
+                // did not actually split edge for t0, so we can still use EID
+                FIndex2i new_info = SplitSegmentAtDistance(EID, t1, VertexSnapTol);
+                FVector2d v = Graph.GetVertex(new_info.A);
+                points.Add(FSegmentPoint{segAB.Project(v), new_info.A});
+            }
+            else
+            {
+                // find t1 was in EID, rebuild in new_eid
+                FSegment2d new_seg = Graph.GetEdgeSegment(new_eid);
+                FVector2d p1 = Intr.Intr.GetSegment1().PointAt(t1);
+                double new_t1 = new_seg.Project(p1);
+                // note: new_t1 may be outside of new_seg due to snapping; in this case the segment will just not be split
+
+                FIndex2i new_info = SplitSegmentAtDistance(new_eid, new_t1, VertexSnapTol);
+                FVector2d v = Graph.GetVertex(new_info.A);
+                points.Add(FSegmentPoint{segAB.Project(v), new_info.A});
+            }
+        }
+    }
+
+    // find or create start and end points
+    if (a_idx == -1)
+    {
+        a_idx = FindExistingVertex(InA);
+    }
+    if (a_idx == -1)
+    {
+        a_idx = Graph.AppendVertex(InA);
+        PointHash.InsertPointUnsafe(a_idx, InA);
+    }
+    if (b_idx == -1)
+    {
+        b_idx = FindExistingVertex(InB);
+    }
+    if (b_idx == -1)
+    {
+        b_idx = Graph.AppendVertex(InB);
+        PointHash.InsertPointUnsafe(b_idx, InB);
+    }
+
+    // add start/end to points list. These may be duplicates but we will sort that out after
+    points.Add(FSegmentPoint{-segAB.Extent, a_idx});
+    points.Add(FSegmentPoint{segAB.Extent, b_idx});
+    // sort by T
+    points.Sort([](const FSegmentPoint& pa, const FSegmentPoint& pb) { return pa.T < pb.T; });
+
+    // connect sequential points, as long as they aren't the same point,
+    // and the segment doesn't already exist
+    for (int k = 0; k < points.Num() - 1; ++k)
+    {
+        int v0 = points[k].VID;
+        int v1 = points[k + 1].VID;
+        if (v0 == v1)
+        {
+            continue;
+        }
+
+        if (Graph.FindEdge(v0, v1) == FDynamicGraph2d::InvalidID)
+        {
+            // sanity check; technically this can happen and still be correct but it's more likely an error case
+            ensureMsgf(FMath::Abs(points[k].T - points[k + 1].T) >= std::numeric_limits<float>::epsilon(), TEXT("insert_segment: different points have same T??"));
+
+            const int EID = Graph.AppendEdge(v0, v1);
+            Directions.Add(EID, v0 < v1);
+        }
+    }
+
+    return true;
+}
+
+FIndex2i FText3DGlyphArrangement::SplitSegmentAtDistance(int InEID, double InDistance, double InTol)
+{
+    FIndex2i ev = Graph.GetEdgeV(InEID);
+    FSegment2d seg = FSegment2d(Graph.GetVertex(ev.A), Graph.GetVertex(ev.B));
+
+    int use_vid = -1;
+    int new_eid = -1;
+    if (InDistance < -(seg.Extent - InTol))
+    {
+        use_vid = ev.A;
+    }
+    else if (InDistance > (seg.Extent - InTol))
+    {
+        use_vid = ev.B;
+    }
+    else
+    {
+        FVector2d Pt = seg.PointAt(InDistance);
+        FDynamicGraph2d::FEdgeSplitInfo splitInfo;
+        EMeshResult result;
+        int CrossingVert = FindExistingVertex(Pt);
+        if (CrossingVert == -1)
+        {
+            result = Graph.SplitEdge(InEID, splitInfo);
+        }
+        else
+        {
+            result = Graph.SplitEdgeWithExistingVertex(InEID, CrossingVert, splitInfo);
+        }
+        ensureMsgf(result == EMeshResult::Ok, TEXT("insert_into_segment: edge split failed?"));
+        use_vid = splitInfo.VNew;
+        new_eid = splitInfo.ENewBN;
+
+        Directions.Add(new_eid, !Directions[InEID]);
+
+        if (CrossingVert == -1)
+        {	// position + track added vertex
+            Graph.SetVertex(use_vid, Pt);
+            PointHash.InsertPointUnsafe(splitInfo.VNew, Pt);
+        }
+    }
+    return FIndex2i(use_vid, new_eid);
+}
+
+int FText3DGlyphArrangement::FindExistingVertex(const FVector2d& InPoint) const
+{
+    return FindNearestVertex(InPoint, VertexSnapTol);
+}
+
+int FText3DGlyphArrangement::FindNearestVertex(const FVector2d& InPoint, double InSearchRadius, int InIgnoreVID) const
+{
+    auto FuncDistSq = [&](int B) { return DistanceSquared(InPoint, Graph.GetVertex(B)); };
+    auto FuncIgnore = [&](int VID) { return VID == InIgnoreVID; };
+    TPair<int, double> found = (InIgnoreVID == -1) ? PointHash.FindNearestInRadius(InPoint, InSearchRadius, FuncDistSq)
+                                                 : PointHash.FindNearestInRadius(InPoint, InSearchRadius, FuncDistSq, FuncIgnore);
+    if (found.Key == PointHash.GetInvalidValue())
+    {
+        return -1;
+    }
+    return found.Key;
+}
+
+bool FText3DGlyphArrangement::FindIntersectingEdges(const FVector2d& InA, const FVector2d& InB, TArray<FIntersection>& OutHits, double InTol) const
+{
+    int num_hits = 0;
+    FVector2d x = FVector2d::Zero(), y = FVector2d::Zero();
+    FVector2d EPerp = UE::Geometry::PerpCW(InB - InA);
+    UE::Geometry::Normalize(EPerp);
+    for (int EID : Graph.EdgeIndices())
+    {
+        Graph.GetEdgeV(EID, x, y);
+        // inlined version of WhichSide with pre-normalized EPerp, to ensure Tolerance is consistent for different edge lengths
+        double SignX = EPerp.Dot(x - InA);
+        double SignY = EPerp.Dot(y - InA);
+        int SideX = (SignX > InTol ? +1 : (SignX < -InTol ? -1 : 0));
+        int SideY = (SignY > InTol ? +1 : (SignY < -InTol ? -1 : 0));
+        if (SideX == SideY && SideX != 0)
+        {
+            continue; // both pts on same side
+        }
+
+        FIntrSegment2Segment2d Intr(FSegment2d(x, y), FSegment2d(InA, InB));
+        Intr.SetIntervalThreshold(InTol);
+        // set a loose DotThreshold as well so almost-parallel segments are treated as parallel;
+        //  otherwise we're more likely to hit later problems when an edge intersects near-overlapping edges at almost the same point
+        // (TODO: detect + handle that case!)
+        Intr.SetDotThreshold(1e-4);
+        if (Intr.Find())
+        {
+            OutHits.Add(FIntersection{EID, SideX, SideY, Intr});
+            num_hits++;
+        }
+    }
+
+    return (num_hits > 0);
+}
+
+bool FText3DGlyphArrangement::FindIntersectingVertices(const FSegment2d &InSegmentAB, int32 InVIDA, int32 InVIDB, TArray<FSegmentPoint>& OutHits, double InTol)
+{
+    int num_hits = 0;
+
+    for (int VID : Graph.VertexIndices())
+    {
+        if (Graph.GetVtxEdgeCount(VID) > 0 || VID == InVIDA || VID == InVIDB) // if it's an existing edge or on the currently added edge, it's not floating so skip it
+        {
+            continue;
+        }
+
+        FVector2d V = Graph.GetVertex(VID);
+        double T;
+        double DSQ = InSegmentAB.DistanceSquared(V, T);
+        if (DSQ < InTol*InTol)
+        {
+            OutHits.Add(FSegmentPoint{ T, VID });
+            num_hits++;
+        }
+    }
+
+    return num_hits > 0;
+}
diff --git a/Engine/Plugins/Experimental/Text3D/Source/Text3D/Private/GeometryBuilders/Text3DGlyphArrangement.h b/Engine/Plugins/Experimental/Text3D/Source/Text3D/Private/GeometryBuilders/Text3DGlyphArrangement.h
new file mode 100644
--- /dev/null
+++ b/Engine/Plugins/Experimental/Text3D/Source/Text3D/Private/GeometryBuilders/Text3DGlyphArrangement.h
@@ -0,0 +1,49 @@
+// Copyright Epic Games, Inc. All Rights Reserved.
+
+#pragma once
+
+#include "Curve/DynamicGraph2.h"
+#include "Spatial/PointHashGrid2.h"
+#include "Intersection/IntrSegment2Segment2.h"
+
+/**
+ * @brief Modified copy of FArrangement2d, has edge directions
+ */
+struct FText3DGlyphArrangement final
+{
+	using FIndex2i = UE::Geometry::FIndex2i;
+	using FSegment2f = UE::Geometry::FSegment2f;
+	using FSegment2d = UE::Geometry::FSegment2d;
+	using FAxisAlignedBox2f = UE::Geometry::FAxisAlignedBox2f;
+	
+    UE::Geometry::FDynamicGraph2d Graph;
+    UE::Geometry::TPointHashGrid2d<int> PointHash;
+    TMap<int, bool> Directions; //bool - from A to B
+	const double VertexSnapTol = 0.001;
+
+	explicit FText3DGlyphArrangement(const FAxisAlignedBox2f& InBoundsHint);
+    void Insert(const FSegment2f& Segment);
+
+protected:
+    struct FSegmentPoint
+    {
+        double T;
+        int VID;
+    };
+
+    struct FIntersection
+    {
+        int EID;
+        int SideX;
+        int SideY;
+        UE::Geometry::FIntrSegment2Segment2d Intr;
+    };
+
+    int InsertPoint(const FVector2d& InPoint, double InTol = 0);
+    bool InsertSegment(FVector2d InA, FVector2d InB, double InTol = 0);
+    UE::Geometry::FIndex2i SplitSegmentAtDistance(int InEID, double InDistance, double InTol);
+    int FindExistingVertex(const FVector2d& InPoint) const;
+    int FindNearestVertex(const FVector2d& InPoint, double InSearchRadius, int InIgnoreVID = -1) const;
+    bool FindIntersectingEdges(const FVector2d& InA, const FVector2d& InB, TArray<FIntersection>& OutHits, double InTol = 0) const;
+    bool FindIntersectingVertices(const FSegment2d& InSegmentAB, int32 InVIDA, int32 InVIDB, TArray<FSegmentPoint>& OutHits, double InTol = 0);
+};
diff --git a/Engine/Plugins/Experimental/Text3D/Source/Text3D/Private/GeometryBuilders/Text3DGlyphContour.cpp b/Engine/Plugins/Experimental/Text3D/Source/Text3D/Private/GeometryBuilders/Text3DGlyphContour.cpp
new file mode 100644
--- /dev/null
+++ b/Engine/Plugins/Experimental/Text3D/Source/Text3D/Private/GeometryBuilders/Text3DGlyphContour.cpp
@@ -0,0 +1,35 @@
+// Copyright Epic Games, Inc. All Rights Reserved.
+
+#include "GeometryBuilders/Text3DGlyphContour.h"
+
+#include "GeometryBuilders/Text3DGlyphPart.h"
+
+FText3DGlyphContour::~FText3DGlyphContour()
+{
+	for (FText3DGlyphPartPtr Part : *this)
+	{
+		Part->Prev.Reset();
+		Part->Next.Reset();
+	}
+}
+
+void FText3DGlyphContour::SetNeighbours()
+{
+	for (int32 Index = 0; Index < Num(); Index++)
+	{
+		FText3DGlyphPartPtr Point = (*this)[Index];
+
+		Point->Prev = (*this)[(Index + Num() - 1) % Num()];
+		Point->Next = (*this)[(Index + 1) % Num()];
+	}
+}
+
+void FText3DGlyphContour::CopyFrom(const FText3DGlyphContour& Other)
+{
+	for (const FText3DGlyphPartConstPtr OtherPoint : Other)
+	{
+		Add(MakeShared<FText3DGlyphPart>(OtherPoint));
+	}
+
+	SetNeighbours();
+}
diff --git a/Engine/Plugins/Experimental/Text3D/Source/Text3D/Private/GeometryBuilders/Text3DGlyphContour.h b/Engine/Plugins/Experimental/Text3D/Source/Text3D/Private/GeometryBuilders/Text3DGlyphContour.h
new file mode 100644
--- /dev/null
+++ b/Engine/Plugins/Experimental/Text3D/Source/Text3D/Private/GeometryBuilders/Text3DGlyphContour.h
@@ -0,0 +1,34 @@
+// Copyright Epic Games, Inc. All Rights Reserved.
+
+#pragma once
+
+#include "Containers/Array.h"
+#include "GeometryBuilders/Text3DGlyphPart.h"
+
+class FText3DGlyphContour final : public TArray<FText3DGlyphPartPtr>
+{
+public:
+	struct FPathEntry
+	{
+		int32 Prev;
+		int32 Next;
+	};
+
+	FText3DGlyphContour() = default;
+	~FText3DGlyphContour();
+
+	int32 Find(const FText3DGlyphPartConstPtr Edge)
+	{
+		return TArray<FText3DGlyphPartPtr>::Find(ConstCastSharedPtr<FText3DGlyphPart>(Edge));
+	}
+
+	/**
+	 * Set Prev and Next in parts.
+	 */
+	void SetNeighbours();
+	/**
+	 * Copy parts from other contour.
+	 * @param Other - Controur from which parts should be copied.
+	 */
+	void CopyFrom(const FText3DGlyphContour& Other);
+};
diff --git a/Engine/Plugins/Experimental/Text3D/Source/Text3D/Private/GeometryBuilders/Text3DGlyphContourList.cpp b/Engine/Plugins/Experimental/Text3D/Source/Text3D/Private/GeometryBuilders/Text3DGlyphContourList.cpp
new file mode 100644
--- /dev/null
+++ b/Engine/Plugins/Experimental/Text3D/Source/Text3D/Private/GeometryBuilders/Text3DGlyphContourList.cpp
@@ -0,0 +1,115 @@
+﻿// Copyright Epic Games, Inc. All Rights Reserved.
+
+#include "GeometryBuilders/Text3DGlyphContourList.h"
+
+#include "GeometryBuilders/Text3DGlyphData.h"
+#include "GeometryBuilders/Text3DGlyphPart.h"
+#include "Math/Vector2D.h"
+#include "Math/Vector.h"
+
+FText3DGlyphContourList::FText3DGlyphContourList(const FText3DGlyphContourList& Other)
+{
+	for (const FText3DGlyphContour& OtherContour : Other)
+	{
+		FText3DGlyphContour& Contour = Add();
+		Contour.CopyFrom(OtherContour);
+	}
+}
+
+void FText3DGlyphContourList::Initialize(const TSharedRef<FText3DGlyphData>& Data)
+{
+	for (FText3DGlyphContour& Contour : *this)
+	{
+		Contour.SetNeighbours();
+
+		for (FText3DGlyphPartPtr Edge : Contour)
+		{
+			Edge->ComputeTangentX();
+		}
+
+		for (const FText3DGlyphPartPtr& Point : Contour)
+		{
+			Point->ComputeSmooth();
+		}
+
+		for (int32 Index = 0; Index < Contour.Num(); Index++)
+		{
+			const FText3DGlyphPartPtr Point = Contour[Index];
+
+			if (!Point->bSmooth && Point->TangentsDotProduct() > 0.f)
+			{
+				const FText3DGlyphPartPtr Curr = Point;
+				const FText3DGlyphPartPtr Prev = Point->Prev;
+
+				const float TangentsCrossProduct = FVector2D::CrossProduct(-Prev->TangentX, Curr->TangentX);
+				const float MinTangentsCrossProduct = 0.9f;
+
+				if (FMath::Abs(TangentsCrossProduct) < MinTangentsCrossProduct)
+				{
+					const float OffsetDefault = 0.01f;
+					const float Offset = FMath::Min3(Prev->Length() / 2.f, Curr->Length() / 2.f, OffsetDefault);
+
+					const FText3DGlyphPartPtr Added = MakeShared<FText3DGlyphPart>();
+					Contour.Insert(Added, Index);
+
+					Prev->Next = Added;
+					Added->Prev = Prev;
+					Added->Next = Curr;
+					Curr->Prev = Added;
+
+					const FVector2D CornerPosition = Curr->Position;
+
+					Curr->Position = CornerPosition + Curr->TangentX * Offset;
+					Added->Position = CornerPosition - Prev->TangentX * Offset;
+
+					Data->AddVertices(1);
+					const int32 VertexID = Data->AddVertex(Added->Position, { 1.f, 0.f }, { -1.f, 0.f, 0.f });
+
+					Added->PathPrev.Add(VertexID);
+					Added->PathNext.Add(VertexID);
+
+					Added->ComputeTangentX();
+
+					Added->ComputeSmooth();
+					Curr->ComputeSmooth();
+				}
+			}
+		}
+
+		for (const FText3DGlyphPartPtr& Point : Contour)
+		{
+			Point->ComputeNormal();
+			Point->ResetInitialPosition();
+		}
+	}
+}
+
+FText3DGlyphContour& FText3DGlyphContourList::Add()
+{
+	AddTail(FText3DGlyphContour());
+	return GetTail()->GetValue();
+}
+
+void FText3DGlyphContourList::Remove(const FText3DGlyphContour& Contour)
+{
+	// Search with comparing pointers
+	for (TDoubleLinkedList<FText3DGlyphContour>::TDoubleLinkedListNode* Node = GetHead(); Node; Node = Node->GetNextNode())
+	{
+		if (&Node->GetValue() == &Contour)
+		{
+			RemoveNode(Node);
+			break;
+		}
+	}
+}
+
+void FText3DGlyphContourList::Reset()
+{
+	for (FText3DGlyphContour& Contour : *this)
+	{
+		for (const FText3DGlyphPartPtr& Part : Contour)
+		{
+			Part->ResetDoneExpand();
+		}
+	}
+}
diff --git a/Engine/Plugins/Experimental/Text3D/Source/Text3D/Private/GeometryBuilders/Text3DGlyphContourList.h b/Engine/Plugins/Experimental/Text3D/Source/Text3D/Private/GeometryBuilders/Text3DGlyphContourList.h
new file mode 100644
--- /dev/null
+++ b/Engine/Plugins/Experimental/Text3D/Source/Text3D/Private/GeometryBuilders/Text3DGlyphContourList.h
@@ -0,0 +1,35 @@
+// Copyright Epic Games, Inc. All Rights Reserved.
+
+#pragma once
+
+#include "Containers/List.h"
+#include "GeometryBuilders/Text3DGlyphContour.h"
+#include "Templates/SharedPointer.h"
+
+class FText3DGlyphData;
+
+class FText3DGlyphContourList final : public TDoubleLinkedList<FText3DGlyphContour>
+{
+public:
+	FText3DGlyphContourList() = default;
+	FText3DGlyphContourList(const FText3DGlyphContourList& Other);
+
+	/**
+	 * Initialize Countours.
+	 * @param Data - Needed to add vertices for split corners.
+	 */
+	void Initialize(const TSharedRef<FText3DGlyphData>& Data);
+
+	/**
+	 * Create contour.
+	 * @return Reference to created contour.
+	 */
+	FText3DGlyphContour& Add();
+	/**
+	 * Remove contour.
+	 * @param Contour - Const reference to contour that should be removed.
+	 */
+	void Remove(const FText3DGlyphContour& Contour);
+
+	void Reset();
+};
diff --git a/Engine/Plugins/Experimental/Text3D/Source/Text3D/Private/GeometryBuilders/Text3DGlyphContourNode.h b/Engine/Plugins/Experimental/Text3D/Source/Text3D/Private/GeometryBuilders/Text3DGlyphContourNode.h
new file mode 100644
--- /dev/null
+++ b/Engine/Plugins/Experimental/Text3D/Source/Text3D/Private/GeometryBuilders/Text3DGlyphContourNode.h
@@ -0,0 +1,32 @@
+// Copyright Epic Games, Inc. All Rights Reserved.
+
+#pragma once
+
+#include "Containers/Array.h"
+#include "Templates/SharedPointer.h"
+#include "Polygon2.h"
+
+using UE::Geometry::FPolygon2f;
+
+struct FText3DGlyphContourNode;
+using TText3DGlyphContourNodeShared = TSharedPtr<FText3DGlyphContourNode>;
+
+struct FText3DGlyphContourNode final
+{
+	explicit FText3DGlyphContourNode(const TSharedPtr<FPolygon2f> ContourIn, const bool bCanHaveIntersectionsIn, const bool bClockwiseIn)
+		: Contour(ContourIn)
+		, bCanHaveIntersections(bCanHaveIntersectionsIn)
+		, bClockwise(bClockwiseIn)
+	{
+	}
+
+	const TSharedPtr<FPolygon2f> Contour;
+
+	// Needed for dividing contours with self-intersections: default value is true, false for parts of divided contours
+	const bool bCanHaveIntersections;
+
+	bool bClockwise;
+
+	// Contours that are inside this contour
+	TArray<TText3DGlyphContourNodeShared> Children;
+};
diff --git a/Engine/Plugins/Experimental/Text3D/Source/Text3D/Private/GeometryBuilders/Text3DGlyphData.cpp b/Engine/Plugins/Experimental/Text3D/Source/Text3D/Private/GeometryBuilders/Text3DGlyphData.cpp
new file mode 100644
--- /dev/null
+++ b/Engine/Plugins/Experimental/Text3D/Source/Text3D/Private/GeometryBuilders/Text3DGlyphData.cpp
@@ -0,0 +1,220 @@
+// Copyright Epic Games, Inc. All Rights Reserved.
+
+#include "GeometryBuilders/Text3DGlyphData.h"
+
+#include "GeometryBuilders/Text3DGlyphPart.h"
+#include "Math/Vector2D.h"
+#include "Math/Vector.h"
+#include "Text3DModule.h"
+
+FText3DGlyphData::FText3DGlyphData(const TSharedPtr<FText3DGlyph>& InGlyph) :
+	Glyph(InGlyph)
+{
+	CurrentGroup = EText3DGroupType::Front;
+	GroupExpand = 0.0f;
+
+	PlannedExtrude = 0.0f;
+	PlannedExpand = 0.0f;
+
+	NormalStart = {0.f, 0.f};
+	NormalEnd = {0.f, 0.f};
+
+	ExtrudeTarget = 0.0f;
+	ExpandTarget = 0.0f;
+
+	DoneExtrude = 0.0f;
+
+	VertexCountBeforeAdd = 0;
+	AddVertexIndex = 0;
+
+	TriangleCountBeforeAdd = 0;
+	AddTriangleIndex = 0;
+}
+
+void FText3DGlyphData::SetCurrentGroup(const EText3DGroupType Type, const float GroupExpandIn)
+{
+	CurrentGroup = Type;
+	check(Glyph.Get());
+
+	FText3DPolygonGroup& Group = Glyph->GetGroups()[static_cast<int32>(Type)];
+	const FMeshDescription& MeshDescription = Glyph->GetMeshDescription();
+
+	Group.FirstVertex = MeshDescription.Vertices().Num();
+	Group.FirstTriangle = MeshDescription.Triangles().Num();
+
+	GroupExpand = GroupExpandIn / FontInverseScale;
+}
+
+void FText3DGlyphData::PrepareSegment(const float PlannedExtrudeIn, const float PlannedExpandIn, const FVector2D NormalStartIn, const FVector2D NormalEndIn)
+{
+	PlannedExtrude = PlannedExtrudeIn;
+	PlannedExpand = PlannedExpandIn / FontInverseScale;
+
+	NormalStart = NormalStartIn;
+	NormalEnd = NormalEndIn;
+}
+
+void FText3DGlyphData::SetTarget(const float ExtrudeTargetIn, const float ExpandTargetIn)
+{
+	ExtrudeTarget = ExtrudeTargetIn;
+	ExpandTarget = ExpandTargetIn;
+}
+
+int32 FText3DGlyphData::AddVertices(const int32 Count)
+{
+	check(Glyph.Get());
+	if (Count <= 0)
+	{
+		return 0;
+	}
+
+	FMeshDescription& MeshDescription = Glyph->GetMeshDescription();
+	FStaticMeshAttributes& MeshAttributes = Glyph->GetStaticMeshAttributes();
+	VertexCountBeforeAdd = MeshDescription.Vertices().Num();
+
+	MeshDescription.ReserveNewVertices(Count);
+	MeshDescription.ReserveNewVertexInstances(Count);
+
+	for (int32 Index = 0; Index < Count; Index++)
+	{
+		const FVertexID Vertex = MeshDescription.CreateVertex();
+		const FVertexInstanceID VertexInstance = MeshDescription.CreateVertexInstance(Vertex);
+		MeshAttributes.GetVertexInstanceColors()[VertexInstance] = FVector4f(1.f, 1.f, 1.f, 1.f);
+	}
+
+	AddVertexIndex = 0;
+	return VertexCountBeforeAdd;
+}
+
+int32 FText3DGlyphData::AddVertex(const FText3DGlyphPartConstPtr& Point, const FVector2D TangentX, const FVector& TangentZ, const FVector2D TextureCoordinates)
+{
+	return AddVertex(Point->Position, TangentX, TangentZ, TextureCoordinates);
+}
+
+int32 FText3DGlyphData::AddVertex(const FVector2D Position, const FVector2D TangentX, const FVector& TangentZ, const FVector2D TextureCoordinates)
+{
+	return AddVertex(GetVector(Position, DoneExtrude + ExtrudeTarget), {0.f, TangentX.X, TangentX.Y}, TangentZ, TextureCoordinates);
+}
+
+int32 FText3DGlyphData::AddVertex(const FVector& Position, const FVector& TangentX, const FVector& TangentZ, const FVector2D TextureCoordinates)
+{
+	check(Glyph.Get());
+	FStaticMeshAttributes& StaticMeshAttributes = Glyph->GetStaticMeshAttributes();
+	const int32 VertexIndex = VertexCountBeforeAdd + AddVertexIndex++;
+	StaticMeshAttributes.GetVertexPositions()[FVertexID(VertexIndex)] = (FVector3f)Position;
+	const FVertexInstanceID Instance(static_cast<uint32>(VertexIndex));
+
+	StaticMeshAttributes.GetVertexInstanceUVs()[Instance] = FVector2f(TextureCoordinates);
+	StaticMeshAttributes.GetVertexInstanceNormals()[Instance] = (FVector3f)TangentZ;
+	StaticMeshAttributes.GetVertexInstanceTangents()[Instance] = (FVector3f)TangentX;
+
+	return VertexIndex;
+}
+
+void FText3DGlyphData::AddTriangles(const int32 Count)
+{
+	check(Glyph.Get());
+	if (Count <= 0)
+	{
+		return;
+	}
+
+	FMeshDescription& MeshDescription = Glyph->GetMeshDescription();
+	TriangleCountBeforeAdd = MeshDescription.Triangles().Num();
+	MeshDescription.ReserveNewTriangles(Count);
+	AddTriangleIndex = 0;
+}
+
+void FText3DGlyphData::AddTriangle(const int32 A, const int32 B, const int32 C)
+{
+	check(Glyph.Get());
+	Glyph->GetMeshDescription().CreateTriangle(FPolygonGroupID(static_cast<int32>(CurrentGroup)), TArray<FVertexInstanceID>({FVertexInstanceID(A), FVertexInstanceID(B), FVertexInstanceID(C)}));
+	AddTriangleIndex++;
+}
+
+float FText3DGlyphData::GetGroupExpand() const
+{
+	return GroupExpand;
+}
+
+float FText3DGlyphData::GetPlannedExtrude() const
+{
+	return PlannedExtrude;
+}
+
+float FText3DGlyphData::GetPlannedExpand() const
+{
+	return PlannedExpand;
+}
+
+void FText3DGlyphData::IncreaseDoneExtrude()
+{
+	DoneExtrude += PlannedExtrude;
+}
+
+FVector FText3DGlyphData::ComputeTangentZ(const FText3DGlyphPartConstPtr& Edge, const float DoneExpand)
+{
+	const FVector2D TangentX = Edge->TangentX;
+
+	const float T = FMath::IsNearlyZero(PlannedExpand) ? 0.0f : DoneExpand / PlannedExpand;
+	const FVector2D Normal = NormalStart * (1.f - T) + NormalEnd * T;
+
+	const FVector2D TangentZ_YZ = FVector2D(TangentX.Y, -TangentX.X) * Normal.X;
+	return FVector(Normal.Y, TangentZ_YZ.X, TangentZ_YZ.Y);
+}
+
+FVector2D FText3DGlyphData::Expanded(const FText3DGlyphPartConstPtr& Point) const
+{
+	// Needed expand value is difference of total expand and point's done expand
+	return Point->Expanded(ExpandTarget - Point->DoneExpand);
+}
+
+void FText3DGlyphData::FillEdge(const FText3DGlyphPartPtr& Edge, const bool bSkipLastTriangle, bool bFlipNormals)
+{
+	const FText3DGlyphPartPtr EdgeA = Edge;
+	const FText3DGlyphPartPtr EdgeB = Edge->Next;
+
+	MakeTriangleFanAlongNormal(EdgeB, EdgeA, bFlipNormals, true);
+	MakeTriangleFanAlongNormal(EdgeA, EdgeB, !bFlipNormals, false);
+
+	if (!bSkipLastTriangle)
+	{
+		MakeTriangleFanAlongNormal(EdgeB, EdgeA, bFlipNormals, false);
+	}
+	else
+	{
+		// Index has to be removed despite last triangle being skipped.
+		// For example when normals intersect and result of expansion of EdgeA and EdgeB is one point -
+		// this point was already covered with last MakeTriangleFanAlongNormal call,
+		// no need to keep it in neighbour point's path.
+		EdgeA->PathNext.RemoveAt(0);
+	}
+
+	// Write done expand
+	EdgeA->DoneExpand = ExpandTarget;
+	EdgeB->DoneExpand = ExpandTarget;
+}
+
+FVector FText3DGlyphData::GetVector(const FVector2D Position, const float Height) const
+{
+	return (FVector(0.f, Position.X, Position.Y) * FontInverseScale + FVector(Height, 0.0f, 0.0f));
+}
+
+void FText3DGlyphData::MakeTriangleFanAlongNormal(const FText3DGlyphPartConstPtr& Cap, const FText3DGlyphPartPtr& Normal, const bool bNormalIsCapNext, const bool bSkipLastTriangle)
+{
+	TArray<int32>& Path = bNormalIsCapNext ? Normal->PathPrev : Normal->PathNext;
+	const int32 Count = Path.Num() - (bSkipLastTriangle ? 2 : 1);
+
+	// Create triangles
+	AddTriangles(Count);
+
+	for (int32 Index = 0; Index < Count; Index++)
+	{
+		AddTriangle((bNormalIsCapNext ? Cap->PathNext : Cap->PathPrev)[0],
+						Path[bNormalIsCapNext ? Index + 1 : Index],
+						Path[bNormalIsCapNext ? Index : Index + 1]);
+	}
+
+	// Remove covered vertices from path
+	Path.RemoveAt(0, Count);
+}
diff --git a/Engine/Plugins/Experimental/Text3D/Source/Text3D/Private/GeometryBuilders/Text3DGlyphData.h b/Engine/Plugins/Experimental/Text3D/Source/Text3D/Private/GeometryBuilders/Text3DGlyphData.h
new file mode 100644
--- /dev/null
+++ b/Engine/Plugins/Experimental/Text3D/Source/Text3D/Private/GeometryBuilders/Text3DGlyphData.h
@@ -0,0 +1,117 @@
+// Copyright Epic Games, Inc. All Rights Reserved.
+
+#pragma once
+
+#include "GeometryBuilders/Text3DGlyph.h"
+#include "GeometryBuilders/Text3DGlyphPart.h"
+#include "Math/MathFwd.h"
+
+/** Used to add vertices and triangles from different classes */
+class FText3DGlyphData final
+{
+public:
+	/**
+	 * Constructor.
+	 * @param InGlyph - Glyph that is being created.
+	 */
+	explicit FText3DGlyphData(const TSharedPtr<FText3DGlyph>& InGlyph);
+
+	/**
+	 * Set group that is being created.
+	 * @param Type - Group type.
+	 * @param GroupExpandIn - Total expand of group (not needed now because intersections are not ready).
+	 */
+	void SetCurrentGroup(const EText3DGroupType Type, const float GroupExpandIn);
+	/**
+	 * Set data needed for bevelling one segment.
+	 * @param PlannedExtrudeIn - Extrude.
+	 * @param PlannedExpandIn - Expand.
+	 * @param NormalStartIn - Normal in the start of segment.
+	 * @param NormalEndIn - Normal in the end of segment.
+	 */
+	void PrepareSegment(const float PlannedExtrudeIn, const float PlannedExpandIn, const FVector2D NormalStartIn, const FVector2D NormalEndIn);
+	/**
+	 * Set data for single bevelling in segment (needed for intersections).
+	 * @param ExtrudeTargetIn - Extrude.
+	 * @param ExpandTargetIn - Expand.
+	 */
+	void SetTarget(const float ExtrudeTargetIn, const float ExpandTargetIn);
+
+	int32 AddVertices(const int32 Count);
+	int32 AddVertex(const FText3DGlyphPartConstPtr& Point, const FVector2D TangentX, const FVector& TangentZ, const FVector2D TextureCoordinates = {0.f, 0.f});
+	int32 AddVertex(const FVector2D Position, const FVector2D TangentX, const FVector& TangentZ, const FVector2D TextureCoordinates = {0.f, 0.f});
+	int32 AddVertex(const FVector& Position, const FVector& TangentX, const FVector& TangentZ, const FVector2D TextureCoordinates);
+
+	void AddTriangles(const int32 Count);
+	void AddTriangle(const int32 A, const int32 B, const int32 C);
+
+	float GetGroupExpand() const;
+
+	float GetPlannedExtrude() const;
+	float GetPlannedExpand() const;
+
+	/**
+	 * Executed after segment was beveled.
+	 */
+	void IncreaseDoneExtrude();
+
+	FVector ComputeTangentZ(const FText3DGlyphPartConstPtr& Edge, const float DoneExpand);
+
+
+	/**
+	 * FPart::Expanded for total expand value Data::ExpandTarget.
+	 * @param Point - Point for which position should be computed.
+	 * @return Computed position.
+	 */
+	FVector2D Expanded(const FText3DGlyphPartConstPtr& Point) const;
+
+	/**
+	 * Make triangulation of edge along paths of it's vertices (from end of previous triangulation to result of points' expansion). Removes covered points' indices from paths.
+	 * @param Edge - Edge that has to be filled.
+	 * @param bSkipLastTriangle - Do not create last triangle (furthest from end of previous triangulation).
+	 * @param bFlipNormals - Reverse the geometry facing direction.
+	 */
+	void FillEdge(const FText3DGlyphPartPtr& Edge, const bool bSkipLastTriangle, bool bFlipNormals = false);
+
+private:
+	TSharedPtr<FText3DGlyph> Glyph;
+
+	EText3DGroupType CurrentGroup;
+	float GroupExpand;
+
+	float PlannedExtrude;
+	float PlannedExpand;
+
+	FVector2D NormalStart;
+	FVector2D NormalEnd;
+
+	float ExtrudeTarget;
+	float ExpandTarget;
+
+	float DoneExtrude;
+
+
+	int32 VertexCountBeforeAdd;
+	int32 AddVertexIndex;
+
+	int32 TriangleCountBeforeAdd;
+	int32 AddTriangleIndex;
+
+
+
+	/**
+	 * Transform position from glyph coordinate system to 3d.
+	 * @param Position - Point position.
+	 * @param Height - Offset in direction orthogonal to surface of front cap.
+	 * @return 3d coordinate.
+	 */
+	FVector GetVector(const FVector2D Position, const float Height) const;
+	/**
+	 * Make triangle fan, called from FData::FillEdge.
+	 * @param Cap - Cap of triangle fan.
+	 * @param Normal - Point, fan will be created along it's normal.
+	 * @param bNormalIsCapNext - Normal is next point after cap or vice versa.
+	 * @param bSkipLastTriangle - See FData::FillEdge.
+	 */
+	void MakeTriangleFanAlongNormal(const FText3DGlyphPartConstPtr& Cap, const FText3DGlyphPartPtr& Normal, const bool bNormalIsCapNext, const bool bSkipLastTriangle);
+};
diff --git a/Engine/Plugins/Experimental/Text3D/Source/Text3D/Private/GeometryBuilders/Text3DGlyphLoader.cpp b/Engine/Plugins/Experimental/Text3D/Source/Text3D/Private/GeometryBuilders/Text3DGlyphLoader.cpp
new file mode 100644
--- /dev/null
+++ b/Engine/Plugins/Experimental/Text3D/Source/Text3D/Private/GeometryBuilders/Text3DGlyphLoader.cpp
@@ -0,0 +1,853 @@
+// Copyright Epic Games, Inc. All Rights Reserved.
+
+#include "GeometryBuilders/Text3DGlyphLoader.h"
+
+#include "BoxTypes.h"
+#include "Curve/DynamicGraph2.h"
+#include "GeometryBuilders/Text3DGlyphArrangement.h"
+#include "GeometryBuilders/Text3DGlyphPart.h"
+#include "Polygon2.h"
+
+using namespace UE::Geometry;
+
+FText3DGlyphLoader::FText3DGlyphLoader(const FT_GlyphSlot InGlyph)
+	: Root(MakeShared<FText3DGlyphContourNode>(nullptr, false, true))
+	, EndIndex(-1)
+	, VertexID(0)
+{
+	check(InGlyph);
+
+	const FT_Outline Outline = InGlyph->outline;
+	const int32 ContourCount = Outline.n_contours;
+
+	for (int32 Index = 0; Index < ContourCount; Index++)
+	{
+		const TSharedPtr<FPolygon2f> Contour = CreateContour(Outline, Index);
+
+		if (Contour.IsValid())
+		{
+			Insert(MakeShared<FText3DGlyphContourNode>(Contour, true, Contour->IsClockwise()), Root);
+		}
+	}
+
+	if (Root->Children.Num() == 0)
+	{
+		return;
+	}
+
+	FixParityAndDivide(Root, false);
+}
+
+TText3DGlyphContourNodeShared FText3DGlyphLoader::GetContourList() const
+{
+	return Root;
+}
+
+
+FText3DGlyphLoader::FLine::FLine(const FVector2D PositionIn)
+	: Position(PositionIn)
+{
+}
+
+void FText3DGlyphLoader::FLine::Add(FText3DGlyphLoader* const Loader)
+{
+	if (!Loader->ProcessedContour.Num() || !(Position - Loader->FirstPosition).IsNearlyZero())
+	{
+		Loader->ProcessedContour.AddHead(FVector2f(Position));
+	}
+}
+
+FText3DGlyphLoader::FCurve::FCurve(const bool bLineIn)
+	: bLine(bLineIn)
+	, StartT(0.f)
+	, EndT(1.f)
+{
+	Loader = nullptr;
+
+	Depth = 0.0f;
+	MaxDepth = 0.0f;
+	First = nullptr;
+	Last = nullptr;
+	bFirstSplit = false;
+	bLastSplit = false;
+}
+
+FText3DGlyphLoader::FCurve::~FCurve()
+{
+}
+
+struct FText3DGlyphLoader::FCurve::FPointData
+{
+	FPointData()
+	{
+		T = 0.0f;
+		Point = nullptr;
+	}
+
+	float T;
+	FVector2D Position;
+	FVector2D Tangent;
+	TDoubleLinkedList<FVector2f>::TDoubleLinkedListNode* Point;
+};
+
+void FText3DGlyphLoader::FCurve::Add(FText3DGlyphLoader* const LoaderIn)
+{
+	Loader = LoaderIn;
+
+	if (bLine)
+	{
+		FLine(Position(StartT)).Add(Loader);
+		return;
+	}
+
+	Depth = 0;
+
+	FPointData Start;
+	Start.T = StartT;
+	Start.Position = Position(Start.T);
+	Start.Tangent = Tangent(Start.T);
+	Loader->ProcessedContour.AddHead(FVector2f(Start.Position));
+	Start.Point = Loader->ProcessedContour.GetHead();
+	First = Start.Point;
+	bFirstSplit = false;
+
+	FPointData End;
+	End.T = EndT;
+	End.Position = Position(End.T);
+	End.Tangent = Tangent(End.T);
+	End.Point = nullptr;
+	Last = End.Point;
+	bLastSplit = false;
+
+
+	ComputeMaxDepth();
+	Split(Start, End);
+}
+
+bool FText3DGlyphLoader::FCurve::OnOneLine(const FVector2D A, const FVector2D B, const FVector2D C)
+{
+	return FMath::IsNearlyZero(FVector2D::CrossProduct((B - A).GetSafeNormal(), (C - A).GetSafeNormal()));
+}
+
+void FText3DGlyphLoader::FCurve::ComputeMaxDepth()
+{
+	// Compute approximate curve length with 4 points
+	const float MinStep = 30.f;
+	const float StepT = 0.333f;
+	float Length = 0.f;
+
+	FVector2D Prev;
+	FVector2D Current = Position(StartT);
+
+	for (float T = StartT + StepT; T < EndT; T += StepT)
+	{
+		Prev = Current;
+		Current = Position(T);
+
+		Length += (Current - Prev).Size();
+	}
+
+	const float MaxStepCount = Length / MinStep;
+	MaxDepth = static_cast<int32>(FMath::Log2(MaxStepCount)) + 1;
+}
+
+void FText3DGlyphLoader::FCurve::Split(const FPointData& Start, const FPointData& End)
+{
+	Depth++;
+	FPointData Middle;
+
+	Middle.T = (Start.T + End.T) / 2.f;
+	Middle.Position = Position(Middle.T);
+	Middle.Tangent = Tangent(Middle.T);
+	Loader->ProcessedContour.InsertNode(FVector2f(Middle.Position), Start.Point);
+	Middle.Point = Start.Point->GetPrevNode();
+
+	CheckPart(Start, Middle);
+	UpdateTangent(&Middle);
+	CheckPart(Middle, End);
+	Depth--;
+}
+
+void FText3DGlyphLoader::FCurve::CheckPart(const FPointData& Start, const FPointData& End)
+{
+	const FVector2D Side = (End.Position - Start.Position).GetSafeNormal();
+
+	if ((FVector2D::DotProduct(Side, Start.Tangent) > FText3DGlyphPart::CosMaxAngleSideTangent &&
+		FVector2D::DotProduct(Side, End.Tangent) > FText3DGlyphPart::CosMaxAngleSideTangent) || Depth >= MaxDepth)
+	{
+		if (!bFirstSplit && Start.Point == First)
+		{
+			bFirstSplit = true;
+			Split(Start, End);
+		}
+		else if (!bLastSplit && End.Point == Last)
+		{
+			bLastSplit = true;
+			Split(Start, End);
+		}
+	}
+	else
+	{
+		Split(Start, End);
+	}
+}
+
+void FText3DGlyphLoader::FCurve::UpdateTangent(FPointData* const Middle)
+{
+}
+
+FText3DGlyphLoader::FQuadraticCurve::FQuadraticCurve(const FVector2D A, const FVector2D B, const FVector2D C)
+	: FCurve(OnOneLine(A, B, C))
+
+	, E(A - 2.f * B + C)
+	, F(-A + B)
+	, G(A)
+{
+}
+
+FVector2D FText3DGlyphLoader::FQuadraticCurve::Position(const float T) const
+{
+	return E * T * T + 2.f * F * T + G;
+}
+
+FVector2D QuadraticCurveTangent(const FVector2D E, const FVector2D F, const float T)
+{
+	const FVector2D Result = E * T + F;
+
+	if (Result.IsNearlyZero())
+	{
+		// Just some vector with non-zero length
+		return {1.f, 0.f};
+	}
+
+	return Result;
+}
+
+FVector2D FText3DGlyphLoader::FQuadraticCurve::Tangent(const float T)
+{
+	return QuadraticCurveTangent(E, F, T).GetSafeNormal();
+}
+
+FText3DGlyphLoader::FCubicCurve::FCubicCurve(const FVector2D A, const FVector2D B, const FVector2D C, const FVector2D D)
+	: FCurve(OnOneLine(A, B, C) && OnOneLine(B, C, D))
+	, E(-A + 3.f * B - 3.f * C + D)
+	, F(A - 2.f * B + C)
+	, G(-A + B)
+	, H(A)
+	, bSharpStart(false)
+	, bSharpMiddle(false)
+	, bSharpEnd(false)
+{
+	if (!bLine)
+	{
+		bSharpStart = G.IsNearlyZero();
+		bSharpEnd = (C - D).IsNearlyZero();
+	}
+}
+
+void FText3DGlyphLoader::FCubicCurve::UpdateTangent(FPointData* const Middle)
+{
+	// In this point curve is not smooth, and  r'(t + 0) / |r'(t + 0)| = -r'(t - 0) / |r'(t - 0)|
+	if (bSharpMiddle)
+	{
+		bSharpMiddle = false;
+		Middle->Tangent *= -1.f;
+	}
+}
+
+FVector2D FText3DGlyphLoader::FCubicCurve::Position(const float T) const
+{
+	return E * T * T * T + 3.f * F * T * T + 3.f * G * T + H;
+}
+
+FVector2D FText3DGlyphLoader::FCubicCurve::Tangent(const float T)
+{
+	FVector2D Result;
+
+	// Using  r' / |r'|  for sharp start and end
+	if (bSharpStart && FMath::IsNearlyEqual(T, StartT))
+	{
+		Result = F;
+	}
+	else if (bSharpEnd && FMath::IsNearlyEqual(T, EndT))
+	{
+		Result = -(E + F);
+	}
+	else
+	{
+		Result = E * T * T + 2.f * F * T + G;
+		bSharpMiddle = Result.IsNearlyZero();
+
+		if (bSharpMiddle)
+		{
+			// Using derivative of quadratic bezier curve (A, B, C) in this point
+			Result = QuadraticCurveTangent(F, G, T);
+		}
+	}
+
+	return Result.GetSafeNormal();
+}
+
+TSharedPtr<FPolygon2f> FText3DGlyphLoader::CreateContour(const FT_Outline Outline, const int32 ContourIndex)
+{
+	const TSharedPtr<FPolygon2f> Contour = ProcessFreetypeOutline(Outline, ContourIndex);
+	return Contour.IsValid() ? RemoveBadPoints(Contour) : Contour;
+}
+
+bool FText3DGlyphLoader::IsInside(const TSharedPtr<const FPolygon2f> ContourA,
+							const TSharedPtr<const FPolygon2f> ContourB) const
+{
+	return ContourB->Contains(ContourA->GetSegmentPoint(0, 0));
+}
+
+void FText3DGlyphLoader::Insert(const TText3DGlyphContourNodeShared NodeA, const TText3DGlyphContourNodeShared NodeB) const
+{
+	TArray<TText3DGlyphContourNodeShared>& ChildrenB = NodeB->Children;
+	const TSharedPtr<const FPolygon2f> ContourA = NodeA->Contour;
+
+	for (int32 ChildBIndex = 0; ChildBIndex < ChildrenB.Num(); ChildBIndex++)
+	{
+		TText3DGlyphContourNodeShared& ChildB = ChildrenB[ChildBIndex];
+		const TSharedPtr<const FPolygon2f> ChildBContour = ChildB->Contour;
+
+		if (IsInside(ContourA, ChildBContour))
+		{
+			Insert(NodeA, ChildB);
+			return;
+		}
+
+		if (IsInside(ChildBContour, ContourA))
+		{
+			// add ChildBContour to list of contours that are inside ContourA
+			TArray<TText3DGlyphContourNodeShared>& ChildrenA = NodeA->Children;
+			ChildrenA.Add(ChildB);
+			// replace ChildBContour with ContourA in list it was before
+			ChildB = NodeA;
+
+			// check if other contours in that list are inside ContourA
+			for (int32 ChildBSiblingIndex = ChildrenB.Num() - 1; ChildBSiblingIndex > ChildBIndex; ChildBSiblingIndex--)
+			{
+				TText3DGlyphContourNodeShared ChildBSibling = ChildrenB[ChildBSiblingIndex];
+
+				if (IsInside(ChildBSibling->Contour, ContourA))
+				{
+					ChildrenA.Add(ChildBSibling);
+					ChildrenB.RemoveAt(ChildBSiblingIndex);
+				}
+			}
+
+			return;
+		}
+	}
+
+	ChildrenB.Add(NodeA);
+}
+
+void FText3DGlyphLoader::FixParityAndDivide(const TText3DGlyphContourNodeShared Node, const bool bClockwiseIn)
+{
+	TArray<TText3DGlyphContourNodeShared>& Children = Node->Children;
+
+	for (int32 ChildIndex = 0; ChildIndex < Children.Num(); ChildIndex++)
+	{
+		TText3DGlyphContourNodeShared Child = Children[ChildIndex];
+		const TSharedPtr<FPolygon2f> ChildContour = Child->Contour;
+
+		// bCanHaveIntersections has default value "true", if it's false then contour was checked for having self-intersections and it's parity was fixed
+		if (!Child->bCanHaveIntersections)
+		{
+			FixParityAndDivide(Child, !bClockwiseIn);
+			continue;
+		}
+
+		// Fix parity
+		if (Child->bClockwise != bClockwiseIn)
+		{
+			ChildContour->Reverse();
+			Child->bClockwise = bClockwiseIn;
+		}
+
+		MakeArrangement(ChildContour);
+		Junctions.Reset();
+
+		// Find first junction
+		if (!FindJunction())
+		{
+			// If no junction was found, contour doesn't have self-intersections
+			FixParityAndDivide(Child, !bClockwiseIn);
+			continue;
+		}
+
+		// Remove this contour from parent's child list
+		Children.RemoveAt(ChildIndex--);
+		// Reset path
+		DividedContourIDs.Reset();
+		// Parts are stored as children of separate root node
+		TText3DGlyphContourNodeShared RootForDetaching = MakeShared<FText3DGlyphContourNode>(nullptr, false, Node->bClockwise);
+		// Processed edges are removed from graph to simplify dividing, vertices that don't belong to any edge are removed too, so a copy is needed
+		const TArray<FVector2f> InitialVertices = CopyVertices();
+
+		bool bAddPoint = true;
+		while (bAddPoint)
+		{
+			AddPoint();
+
+			if (VertexID == -1)
+			{
+				// This is bad contour (not a closed loop)
+				VertexID = Junctions.Last();
+				DetachBadContour();
+			}
+			else
+			{
+				const int32 RepeatedJunctionIndex = Junctions.Find(VertexID);
+				if (RepeatedJunctionIndex == INDEX_NONE)
+				{
+					continue;
+				}
+
+				// If it is, we made a loop that should be separated
+				DetachFinishedContour(RepeatedJunctionIndex, InitialVertices, RootForDetaching);
+			}
+
+			// If path contains no junctions and current vertex is not a junction
+			if (Junctions.Num() == 0 && !Arrangement->Graph.IsJunctionVertex(VertexID))
+			{
+				if (!FindJunction())
+				{
+					// If there are no junctions in graph, it contains last contour that should be separated
+					if (!FindRegular())
+					{
+						// Graph is empty, done
+						bAddPoint = false;
+					}
+					else
+					{
+						// Add non-junction vertex as if it was a junction
+						AddPoint(true);
+					}
+				}
+			}
+		}
+
+		RemoveUnneededNodes(RootForDetaching);
+		MergeRootForDetaching(Child, RootForDetaching, Node);
+	}
+}
+
+TSharedPtr<FPolygon2f> FText3DGlyphLoader::ProcessFreetypeOutline(const FT_Outline Outline, const int32 ContourIndex)
+{
+	TSharedPtr<FPolygon2f> Contour;
+
+	const int32 StartIndex = EndIndex + 1;
+	EndIndex = Outline.contours[ContourIndex];
+	const int32 ContourLength = EndIndex - StartIndex + 1;
+
+	if (ContourLength < 3)
+	{
+		return Contour;
+	}
+
+	ProcessedContour.Empty();
+
+	const FT_Vector* const Points = Outline.points + StartIndex;
+	auto ToFVector2D = [Points](const int32 Index)
+	{
+		const FT_Vector Point = Points[Index];
+		return FVector2D(Point.x, Point.y);
+	};
+
+	FVector2D Prev;
+	FVector2D Current = ToFVector2D(ContourLength - 1);
+	FVector2D Next = ToFVector2D(0);
+	FVector2D NextNext = ToFVector2D(1);
+
+
+	const char* const Tags = Outline.tags + StartIndex;
+	auto Tag = [Tags](int32 Index)
+	{
+		return FT_CURVE_TAG(Tags[Index]);
+	};
+
+	int32 TagPrev = 0;
+	int32 TagCurrent = Tag(ContourLength - 1);
+	int32 TagNext = Tag(0);
+
+
+	FVector2D& FirstPositionLocal = FirstPosition;
+	TDoubleLinkedList<FVector2f>& ProcessedContourLocal = ProcessedContour;
+	auto ContourIsBad = [&FirstPositionLocal, &ProcessedContourLocal](const FVector2D Point)
+	{
+		if (ProcessedContourLocal.Num() == 0)
+		{
+			FirstPositionLocal = Point;
+			return false;
+		}
+
+		return (Point - FVector2D(ProcessedContourLocal.GetHead()->GetValue())).IsNearlyZero();
+	};
+
+	for (int32 Index = 0; Index < ContourLength; Index++)
+	{
+		const int32 NextIndex = (Index + 1) % ContourLength;
+
+		Prev = Current;
+		Current = Next;
+		Next = NextNext;
+		NextNext = ToFVector2D((NextIndex + 1) % ContourLength);
+
+		TagPrev = TagCurrent;
+		TagCurrent = TagNext;
+		TagNext = Tag(NextIndex);
+
+		if (TagCurrent == FT_Curve_Tag_On)
+		{
+			if (TagNext == FT_Curve_Tag_Cubic || TagNext == FT_Curve_Tag_Conic)
+			{
+				continue;
+			}
+
+			if (ContourIsBad(Current))
+			{
+				return Contour;
+			}
+
+			if (TagNext == FT_Curve_Tag_On && (Current - Next).IsNearlyZero())
+			{
+				continue;
+			}
+
+			FLine(Current).Add(this);
+		}
+		else if (TagCurrent == FT_Curve_Tag_Conic)
+		{
+			FVector2D A;
+
+			if (TagPrev == FT_Curve_Tag_On)
+			{
+				if (ContourIsBad(Prev))
+				{
+					return Contour;
+				}
+
+				A = Prev;
+			}
+			else
+			{
+				A = (Prev + Current) / 2.f;
+			}
+
+			FQuadraticCurve(A, Current, TagNext == FT_Curve_Tag_Conic ? (Current + Next) / 2.f : Next).Add(this);
+		}
+		else if (TagCurrent == FT_Curve_Tag_Cubic && TagNext == FT_Curve_Tag_Cubic)
+		{
+			if (ContourIsBad(Prev))
+			{
+				return Contour;
+			}
+
+			FCubicCurve(Prev, Current, Next, NextNext).Add(this);
+		}
+	}
+
+	if (ProcessedContour.Num() < 3)
+	{
+		return Contour;
+	}
+
+
+	Contour = MakeShared<FPolygon2f>();
+
+	for (const TDoubleLinkedList<FVector2f>::TDoubleLinkedListNode* Point = ProcessedContour.GetTail(); Point; Point =
+		 Point->GetPrevNode())
+	{
+		Contour->AppendVertex(Point->GetValue());
+	}
+
+	return Contour;
+}
+
+bool FText3DGlyphLoader::IsBadNormal(const TSharedPtr<FPolygon2f> Contour, const int32 Point) const
+{
+	FVector2f ToPrev;
+	FVector2f ToNext;
+
+	Contour->NeighbourVectors(Point, ToPrev, ToNext, true);
+	return FMath::IsNearlyZero(1.0f - FVector2D::DotProduct(FVector2D(ToPrev), FVector2D(ToNext)));
+}
+
+bool FText3DGlyphLoader::MergedWithNext(const TSharedPtr<FPolygon2f> Contour, const int32 Point) const
+{
+	const int32 Current = Point;
+	const int32 Next = (Point + 1) % Contour->VertexCount();
+
+	return FVector2D((*Contour)[Next] - (*Contour)[Current]).IsNearlyZero();
+}
+
+TSharedPtr<FPolygon2f> FText3DGlyphLoader::RemoveBadPoints(const TSharedPtr<FPolygon2f> Contour) const
+{
+	int32 First = 0;
+
+	for (int32 Point = First;;)
+	{
+		bool bPointRemoved = false;
+
+		while (IsBadNormal(Contour, Point) || MergedWithNext(Contour, Point))
+		{
+			bPointRemoved = true;
+
+			if (Contour->VertexCount() < 4)
+			{
+				return TSharedPtr<FPolygon2f>();
+			}
+
+			Contour->RemoveVertex(Point);
+			const int32 Current = (Point - 1 + Contour->VertexCount()) % Contour->VertexCount();
+
+			if (Point == 0)
+			{
+				First = Current;
+			}
+
+			Point = Current;
+		}
+
+		if (!bPointRemoved)
+		{
+			Point = (Point + 1) % Contour->VertexCount();
+
+			if (Point == First)
+			{
+				break;
+			}
+		}
+	}
+
+	return Contour;
+}
+
+void FText3DGlyphLoader::MakeArrangement(const TSharedPtr<const FPolygon2f> Contour)
+{
+	FAxisAlignedBox2f Box = FAxisAlignedBox2f::Empty();
+
+	for (const FVector2f& Vertex : Contour->GetVertices())
+	{
+		Box.Contain(Vertex);
+	}
+
+	Arrangement = MakeShared<FText3DGlyphArrangement>(Box);
+
+	for (FSegment2f Edge : Contour->Segments())
+	{
+		Arrangement->Insert(Edge);
+	}
+}
+
+bool FText3DGlyphLoader::FindJunction()
+{
+	FDynamicGraph2d& Graph = Arrangement->Graph;
+
+	for (int32 VID = 0; VID < Graph.MaxVertexID(); VID++)
+	{
+		if (Graph.IsJunctionVertex(VID))
+		{
+			VertexID = VID;
+			return true;
+		}
+	}
+
+	VertexID = -1;
+	return false;
+}
+
+bool FText3DGlyphLoader::FindRegular()
+{
+	FDynamicGraph2d& Graph = Arrangement->Graph;
+
+	for (int32 VID = 0; VID < Graph.MaxVertexID(); VID++)
+	{
+		if (Graph.IsRegularVertex(VID))
+		{
+			VertexID = VID;
+			return true;
+		}
+	}
+
+	VertexID = -1;
+	return false;
+}
+
+TArray<FVector2f> FText3DGlyphLoader::CopyVertices() const
+{
+	FDynamicGraph2d& Graph = Arrangement->Graph;
+	TArray<FVector2f> Vertices;
+	const int32 MaxVID = Graph.MaxVertexID();
+	Vertices.Reserve(MaxVID);
+
+	for (int32 VID = 0; VID < MaxVID; VID++)
+	{
+		const FVector2d Vertex = Graph.GetVertex(VID);
+		Vertices.Add({static_cast<float>(Vertex.X), static_cast<float>(Vertex.Y)});
+	}
+
+	return Vertices;
+}
+
+void FText3DGlyphLoader::DetachBadContour()
+{
+	const int32 JunctionIndexInContour = FindStartOfDetachedContour(Junctions.Num() - 1);
+	RemoveDetachedContourFromPath(JunctionIndexInContour);
+}
+
+void FText3DGlyphLoader::DetachFinishedContour(const int32 RepeatedJunctionIndex, const TArray<FVector2f>& InitialVertices,
+										 const TText3DGlyphContourNodeShared RootForDetaching)
+{
+	const int32 JunctionIndexInContour = FindStartOfDetachedContour(RepeatedJunctionIndex);
+
+	// Create contour from copied vertices
+	TSharedPtr<FPolygon2f> FinishedContour = MakeShared<FPolygon2f>();
+
+	for (int32 ID = JunctionIndexInContour; ID < DividedContourIDs.Num(); ID++)
+	{
+		FinishedContour->AppendVertex(InitialVertices[DividedContourIDs[ID]]);
+	}
+
+	RemoveDetachedContourFromPath(JunctionIndexInContour);
+
+	const TText3DGlyphContourNodeShared FinishedContourNode = MakeShared<FText3DGlyphContourNode>(FinishedContour,
+																			false,
+																			FinishedContour->IsClockwise());
+	Insert(FinishedContourNode, RootForDetaching);
+}
+
+int32 FText3DGlyphLoader::FindStartOfDetachedContour(const int32 RepeatedJunctionIndex)
+{
+	// Remove from list of junctions all junctions that belong to detached loop
+	Junctions.SetNum(RepeatedJunctionIndex);
+	// Find index in path of first (and last) vertex of loop
+	const int32 JunctionIndexInContour = DividedContourIDs.Find(VertexID);
+	return JunctionIndexInContour;
+}
+
+void FText3DGlyphLoader::RemoveDetachedContourFromPath(const int32 JunctionIndexInContour)
+{
+	DividedContourIDs.SetNum(JunctionIndexInContour);
+}
+
+void FText3DGlyphLoader::AddPoint(const bool bForceJunction)
+{
+	DividedContourIDs.Add(VertexID);
+	FDynamicGraph2d& Graph = Arrangement->Graph;
+	int32 OutgoingEdge = 0;
+
+	if (Graph.IsJunctionVertex(VertexID) || bForceJunction)
+	{
+		Junctions.Add(VertexID);
+		// Select specific path direction
+		OutgoingEdge = FindOutgoing(Junctions.Last());
+	}
+	else if (Graph.IsVertex(VertexID))
+	{
+		// This should be the only possible direction
+		OutgoingEdge = *Graph.VtxEdgesItr(VertexID).begin();
+	}
+	else
+	{
+		// We have no possible directions, this contour is bad
+		VertexID = -1;
+		return;
+	}
+
+	// Move to next graph node
+	VertexID = GetOtherEdgeVertex(OutgoingEdge, VertexID);
+	Graph.RemoveEdge(OutgoingEdge, true);
+}
+
+void FText3DGlyphLoader::RemoveUnneededNodes(const TText3DGlyphContourNodeShared Node) const
+{
+	TArray<TText3DGlyphContourNodeShared>& Children = Node->Children;
+
+	for (int32 ChildIndex = 0; ChildIndex < Children.Num();)
+	{
+		const TText3DGlyphContourNodeShared Child = Children[ChildIndex];
+
+		if (Child->bClockwise == Node->bClockwise)
+		{
+			Children.RemoveAt(ChildIndex);
+		}
+		else
+		{
+			ChildIndex++;
+			RemoveUnneededNodes(Child);
+		}
+	}
+}
+
+void FText3DGlyphLoader::MergeRootForDetaching(const TText3DGlyphContourNodeShared RemovedNode,
+										 const TText3DGlyphContourNodeShared RootForDetaching,
+										 const TText3DGlyphContourNodeShared RemovedNodeParent) const
+{
+	for (const TText3DGlyphContourNodeShared& Child : RemovedNode->Children)
+	{
+		Insert(Child, RootForDetaching);
+	}
+
+	for (const TText3DGlyphContourNodeShared& Child : RootForDetaching->Children)
+	{
+		RemovedNodeParent->Children.Add(Child);
+	}
+}
+
+bool FText3DGlyphLoader::IsOutgoing(const int32 Junction, const int32 EID) const
+{
+	return (Arrangement->Graph.GetEdgeV(EID).A == Junction) == Arrangement->Directions[EID];
+}
+
+int32 FText3DGlyphLoader::FindOutgoing(const int32 Junction) const
+{
+	// Get list of all edges that have this vertex
+	TArray<int32> Edges;
+	SortedVtxEdges(Junction, Edges);
+
+	int32 Current = 0;
+	int32 Next = Edges[0];
+
+	for (int32 EdgeIndex = 0; EdgeIndex < Edges.Num(); EdgeIndex++)
+	{
+		Current = Next;
+		Next = Edges[(EdgeIndex + 1) % Edges.Num()];
+
+		// Needed direction is the one that is outgoing and it's clockwise neighbour is not
+		if (!IsOutgoing(Junction, Current) && IsOutgoing(Junction, Next))
+		{
+			return Next;
+		}
+	}
+
+	return Edges[0];
+}
+
+int32 FText3DGlyphLoader::GetOtherEdgeVertex(const int32 EID, const int32 VID) const
+{
+	const FIndex2i Edge = Arrangement->Graph.GetEdgeV(EID);
+	return (Edge.A == VID) ? Edge.B : ((Edge.B == VID) ? Edge.A : FDynamicGraph::InvalidID);
+}
+
+void FText3DGlyphLoader::SortedVtxEdges(const int32 VID, TArray<int32>& Sorted) const
+{
+	FDynamicGraph2d& Graph = Arrangement->Graph;
+	Sorted.Reserve(Graph.GetVtxEdgeCount(VID));
+
+	for (int32 EID : Graph.VtxEdgesItr(VID))
+	{
+		Sorted.Add(EID);
+	}
+
+	const FVector2d V = Graph.GetVertex(VID);
+	Algo::SortBy(Sorted,
+				 [&](int32 EID)
+				 {
+					 const int32 NbrVID = GetOtherEdgeVertex(EID, VID);
+					 const FVector2d D = Graph.GetVertex(NbrVID) - V;
+					 return TMathUtil<double>::Atan2Positive(D.Y, D.X);
+				 });
+}
diff --git a/Engine/Plugins/Experimental/Text3D/Source/Text3D/Private/GeometryBuilders/Text3DGlyphLoader.h b/Engine/Plugins/Experimental/Text3D/Source/Text3D/Private/GeometryBuilders/Text3DGlyphLoader.h
new file mode 100644
--- /dev/null
+++ b/Engine/Plugins/Experimental/Text3D/Source/Text3D/Private/GeometryBuilders/Text3DGlyphLoader.h
@@ -0,0 +1,235 @@
+// Copyright Epic Games, Inc. All Rights Reserved.
+
+#pragma once
+
+#include "Containers/List.h"
+#include "GeometryBuilders/Text3DGlyphContourNode.h"
+#include "Math/Vector2D.h"
+#include "Polygon2.h"
+#include "Templates/SharedPointer.h"
+#include "Text3DModule.h"
+
+using UE::Geometry::FPolygon2f;
+
+/** Creates glyph geometry data from a font glyphs */
+class FText3DGlyphLoader final
+{
+public:
+	explicit FText3DGlyphLoader(const FT_GlyphSlot InGlyph);
+
+	TText3DGlyphContourNodeShared GetContourList() const;
+
+private:
+	class FLine final
+	{
+	public:
+		FLine(const FVector2D PositionIn);
+		void Add(FText3DGlyphLoader* const Loader);
+
+	private:
+		const FVector2D Position;
+	};
+
+	class FCurve
+	{
+	public:
+		FCurve(const bool bLineIn);
+		virtual ~FCurve();
+
+		void Add(FText3DGlyphLoader* const LoaderIn);
+
+	protected:
+		struct FPointData;
+
+
+		/**
+		 * Check if 3 points are on one line.
+		 * @param A - Point A.
+		 * @param B - Point B.
+		 * @param C - Point C.
+		 * @return Are points on one line?
+		 */
+		static bool OnOneLine(const FVector2D A, const FVector2D B, const FVector2D C);
+
+
+		/** Is curve a line? */
+		const bool bLine;
+		FText3DGlyphLoader* Loader;
+
+		const float StartT;
+		const float EndT;
+
+	private:
+		int32 Depth;
+		int32 MaxDepth;
+
+		TDoubleLinkedList<FVector2f>::TDoubleLinkedListNode* First;
+		TDoubleLinkedList<FVector2f>::TDoubleLinkedListNode* Last;
+
+		/** Needed to make additional splits near start and end of curve */
+		bool bFirstSplit;
+		bool bLastSplit;
+
+
+		/**
+		 * Compute max depth (depends on curve length, step is fixed).
+		 */
+		void ComputeMaxDepth();
+		void Split(const FPointData& Start, const FPointData& End);
+		void CheckPart(const FPointData& Start, const FPointData& End);
+		virtual void UpdateTangent(FPointData* const Middle);
+
+		virtual FVector2D Position(const float T) const = 0;
+		virtual FVector2D Tangent(const float T) = 0;
+	};
+
+	class FQuadraticCurve final : public FCurve
+	{
+	public:
+		FQuadraticCurve(const FVector2D A, const FVector2D B, const FVector2D C);
+
+	private:
+		const FVector2D E;
+		const FVector2D F;
+		const FVector2D G;
+
+
+		FVector2D Position(const float T) const override;
+		FVector2D Tangent(const float T) override;
+	};
+
+	class FCubicCurve final : public FCurve
+	{
+	public:
+		FCubicCurve(const FVector2D A, const FVector2D B, const FVector2D C, const FVector2D D);
+
+	private:
+		const FVector2D E;
+		const FVector2D F;
+		const FVector2D G;
+		const FVector2D H;
+
+		/** Sharp means that curve derivative has zero length, it's actually sharp only in middle case */
+		bool bSharpStart;
+		bool bSharpMiddle;
+		bool bSharpEnd;
+
+
+		void UpdateTangent(FPointData* const Middle) override;
+
+		FVector2D Position(const float t) const override;
+		FVector2D Tangent(const float t) override;
+	};
+
+
+	TText3DGlyphContourNodeShared Root;
+	int32 EndIndex;
+	TDoubleLinkedList<FVector2f> ProcessedContour;
+	FVector2D FirstPosition;
+	TSharedPtr<struct FText3DGlyphArrangement> Arrangement;
+	// List of junctions in path
+	TArray<int32> Junctions;
+	// Current vertex that should be added to path
+	int32 VertexID;
+	// Path
+	TArray<int32> DividedContourIDs;
+
+
+
+	TSharedPtr<FPolygon2f> CreateContour(const FT_Outline Outline, const int32 ContourIndex);
+
+	/**
+	 * Checks if ContourA is inside ContourB.
+	 * @param ContourA
+	 * @param ContourB
+	 */
+	bool IsInside(const TSharedPtr<const FPolygon2f> ContourA, const TSharedPtr<const FPolygon2f> ContourB) const;
+
+	/**
+	 * Insert NodeA inside NodeB.
+	 * @param NodeA
+	 * @param NodeB
+	 */
+	void Insert(const TText3DGlyphContourNodeShared NodeA, const TText3DGlyphContourNodeShared NodeB) const;
+	/**
+	 * Reverse contour if it's initial parity differs from the one it should have, find self-intersections and, if found, divide to contours without self-intersections.
+	 * @param Node - Function is called recursively to fix all contours inside Node->Contour.
+	 * @param ClockwiseIn - The parity that contours listed in Node->Nodes should have.
+	 */
+	void FixParityAndDivide(const TText3DGlyphContourNodeShared Node, const bool bClockwiseIn);
+
+	TSharedPtr<FPolygon2f> ProcessFreetypeOutline(const FT_Outline Outline, const int32 ContourIndex);
+	TSharedPtr<FPolygon2f> RemoveBadPoints(const TSharedPtr<FPolygon2f> Contour) const;
+
+	bool IsBadNormal(const TSharedPtr<FPolygon2f> Contour, const int32 Point) const;
+	bool MergedWithNext(const TSharedPtr<FPolygon2f> Contour, const int32 Point) const;
+
+	void MakeArrangement(const TSharedPtr<const FPolygon2f> Contour);
+	/**
+	 * Find vertex that belongs to more then 2 edges and write it's ID to VertexID member.
+	 * @return Was such vertex found?
+	 */
+	bool FindJunction();
+	/**
+	 * Find vertex that belongs to 2 edges and write it's ID to VertexID member.
+	 * @return Was such vertex found?
+	 */
+	bool FindRegular();
+	TArray<FVector2f> CopyVertices() const;
+
+	/**
+	 * Separate bad contour (the one that is not a closed loop) from path.
+	 */
+	void DetachBadContour();
+	/**
+	 * Separate closed loop from path.
+	 * @param RepeatedJunctionIndex - Index of start (and end) of loop in list of junctions.
+	 * @param InitialVertices - Previously copied graph vertices.
+	 * @param RootForDetaching - Separate root node for detaching.
+	 */
+	void DetachFinishedContour(const int32 RepeatedJunctionIndex, const TArray<FVector2f>& InitialVertices, const TText3DGlyphContourNodeShared RootForDetaching);
+
+	/**
+	 * Find start of contour that is detached, remove it from list of junctions.
+	 * @param RepeatedJunctionIndex - Index of start (and end) of loop in list of junctions.
+	 * @return Index in path of first (and last) vertex of loop.
+	 */
+	int32 FindStartOfDetachedContour(const int32 RepeatedJunctionIndex);
+	/**
+	 * Remove separated contour from path.
+	 * @param JunctionIndexInContour - Index in path of first (and last) vertex of loop.
+	 */
+	void RemoveDetachedContourFromPath(const int32 JunctionIndexInContour);
+
+	/**
+	 * Add vertex to path
+	 * @param bForceJunction - Needed for separating last loop, it has no junctions so any other vertex has to be treated as one.
+	 */
+	void AddPoint(const bool bForceJunction = false);
+	/**
+	 * Check what nodes have same parity as their parents and remove them, they cover the area that already is covered by parent.
+	 * @param Node - Processed node.
+	 */
+	void RemoveUnneededNodes(const TText3DGlyphContourNodeShared Node) const;
+	/**
+	 * Merge separate root to regular one.
+	 * @param RemovedNode - Contour that was divided.
+	 * @param RootForDetaching - Separate root that stores result of division.
+	 * @param RemovedNodeParent - Parent of contour that was divided.
+	 */
+	void MergeRootForDetaching(const TText3DGlyphContourNodeShared RemovedNode, const TText3DGlyphContourNodeShared RootForDetaching, const TText3DGlyphContourNodeShared RemovedNodeParent) const;
+
+	/**
+	 * Choose next path direction in junction vertex.
+	 * @param Junction - Vertex ID of junction.
+	 * @return Edge ID of chosen direction.
+	 */
+	int32 FindOutgoing(const int32 Junction) const;
+	bool IsOutgoing(const int32 Junction, const int32 EID) const;
+
+	// Copied from FDynamicGraph2d because it's private there
+	int32 GetOtherEdgeVertex(const int32 EID, const int32 VID) const;
+
+	// Fixed version of FDynamicGraph2d::SortedVtxEdges
+	void SortedVtxEdges(const int32 VID, TArray<int32>& Sorted) const;
+};
diff --git a/Engine/Plugins/Experimental/Text3D/Source/Text3D/Private/GeometryBuilders/Text3DGlyphMeshBuilder.cpp b/Engine/Plugins/Experimental/Text3D/Source/Text3D/Private/GeometryBuilders/Text3DGlyphMeshBuilder.cpp
new file mode 100644
--- /dev/null
+++ b/Engine/Plugins/Experimental/Text3D/Source/Text3D/Private/GeometryBuilders/Text3DGlyphMeshBuilder.cpp
@@ -0,0 +1,666 @@
+// Copyright Epic Games, Inc. All Rights Reserved.
+
+#include "GeometryBuilders/Text3DGlyphMeshBuilder.h"
+
+#include "ConstrainedDelaunay2.h"
+#include "GeometryBuilders/Text3DGlyphContour.h"
+#include "GeometryBuilders/Text3DGlyphContourList.h"
+#include "GeometryBuilders/Text3DGlyphData.h"
+#include "GeometryBuilders/Text3DGlyphPart.h"
+
+using namespace UE::Geometry;
+
+FText3DGlyphMeshBuilder::FText3DGlyphMeshBuilder() :
+	Glyph(MakeShared<FText3DGlyph>()),
+	Data(MakeShared<FText3DGlyphData>(Glyph))
+{
+}
+
+void FText3DGlyphMeshBuilder::CreateMeshes(const TText3DGlyphContourNodeShared& Root, const float Extrude, const float Bevel, const EText3DBevelType Type, const int32 BevelSegments, const bool bOutline, const float OutlineExpand)
+{
+	CreateFrontMesh(Root, bOutline, OutlineExpand);
+	if (Contours->Num() == 0)
+	{
+		return;
+	}
+
+	const bool bFlipNormals = FMath::Sign(Data->GetPlannedExpand()) < 0;
+	const float BevelLocal = bOutline ? 0.f : Bevel;
+	CreateBevelMesh(BevelLocal, Type, BevelSegments);
+	CreateExtrudeMesh(Extrude, BevelLocal, Type, bFlipNormals);
+}
+
+void FText3DGlyphMeshBuilder::SetFrontAndBevelTextureCoordinates(const float Bevel)
+{
+	EText3DGroupType GroupType = FMath::IsNearlyZero(Bevel) ? EText3DGroupType::Front : EText3DGroupType::Bevel;
+	int32 GroupIndex = static_cast<int32>(GroupType);
+
+	FBox2f Box;
+	TText3DGroupList& Groups = Glyph->GetGroups();
+
+	const int32 FirstVertex = Groups[GroupIndex].FirstVertex;
+	const int32 LastVertex = Groups[GroupIndex + 1].FirstVertex;
+
+	TVertexAttributesConstRef<FVector3f> Positions = Glyph->GetStaticMeshAttributes().GetVertexPositions();
+
+	const FVector3f& FirstPosition = Positions[FVertexID(FirstVertex)];
+	const FVector2f PositionFlat = { FirstPosition.Y, FirstPosition.Z };
+
+	Box.Min = PositionFlat;
+	Box.Max = PositionFlat;
+
+
+	for (int32 VertexIndex = FirstVertex + 1; VertexIndex < LastVertex; VertexIndex++)
+	{
+		const FVector3f& Position = Positions[FVertexID(VertexIndex)];
+
+		Box.Min.X = FMath::Min(Box.Min.X, Position.Y);
+		Box.Min.Y = FMath::Min(Box.Min.Y, Position.Z);
+		Box.Max.X = FMath::Max(Box.Max.X, Position.Y);
+		Box.Max.Y = FMath::Max(Box.Max.Y, Position.Z);
+	}
+
+	FStaticMeshAttributes& StaticMeshAttributes = Glyph->GetStaticMeshAttributes();
+	TVertexAttributesRef<FVector3f> VertexPositions = StaticMeshAttributes.GetVertexPositions();
+	TVertexInstanceAttributesRef<FVector2f> VertexInstanceUVs = StaticMeshAttributes.GetVertexInstanceUVs();
+
+	auto SetTextureCoordinates = [Groups, VertexPositions, VertexInstanceUVs, &Box](const EText3DGroupType Type)
+	{
+		const int32 TypeFirstVertex = Groups[static_cast<int32>(Type)].FirstVertex;
+		const int32 TypeLastVertex = Groups[static_cast<int32>(Type) + 1].FirstVertex;
+
+		for (int32 Index = TypeFirstVertex; Index < TypeLastVertex; Index++)
+		{
+			const FVector Position = (FVector)VertexPositions[FVertexID(Index)];
+			const FVector2f TextureCoordinate = (FVector2f(Position.Y, Position.Z) - Box.Min) / Box.Max;
+			VertexInstanceUVs[FVertexInstanceID(Index)] = { TextureCoordinate.X, 1.f - TextureCoordinate.Y };
+		}
+	};
+
+	SetTextureCoordinates(EText3DGroupType::Front);
+
+	if (!FMath::IsNearlyZero(Bevel))
+	{
+		SetTextureCoordinates(EText3DGroupType::Bevel);
+	}
+}
+
+void FText3DGlyphMeshBuilder::MirrorGroups(const float Extrude)
+{
+	// No bevels without extrude
+	if (!FMath::IsNearlyZero(Extrude))
+	{
+		MirrorGroup(EText3DGroupType::Front, EText3DGroupType::Back, Extrude);
+		MirrorGroup(EText3DGroupType::Bevel, EText3DGroupType::Bevel, Extrude);
+	}
+}
+
+void FText3DGlyphMeshBuilder::BuildMesh(UStaticMesh* StaticMesh, class UMaterial* DefaultMaterial)
+{
+	Glyph->Build(StaticMesh, DefaultMaterial);
+}
+
+void FText3DGlyphMeshBuilder::CreateFrontMesh(const TText3DGlyphContourNodeShared& Root, const bool bOutline, const float& OutlineExpand)
+{
+	int32 VertexCount = 0;
+	AddToVertexCount(Root, VertexCount);
+
+	Data->SetCurrentGroup(EText3DGroupType::Front, bOutline ? OutlineExpand : 0.f);
+	Data->SetTarget(0.f, 0.f);
+	Contours = MakeShared<FText3DGlyphContourList>();
+
+	int32 VertexIndex = Data->AddVertices(VertexCount);
+	TriangulateAndConvert(Root, VertexIndex, bOutline);
+
+	Contours->Initialize(Data);
+
+	if (bOutline)
+	{
+		MakeOutline(OutlineExpand);
+	}
+}
+
+void FText3DGlyphMeshBuilder::CreateBevelMesh(const float Bevel, const EText3DBevelType Type, const int32 BevelSegments)
+{
+	Data->SetCurrentGroup(EText3DGroupType::Bevel, Bevel);
+
+	if (FMath::IsNearlyZero(Bevel))
+	{
+		return;
+	}
+
+	switch (Type)
+	{
+	case EText3DBevelType::Linear:
+	{
+		BevelLinearWithSegments(Bevel, Bevel, BevelSegments, FVector2D(1.f, -1.f).GetSafeNormal());
+		break;
+	}
+	case EText3DBevelType::Convex:
+	{
+		BevelCurve(HALF_PI, BevelSegments, [Bevel](const float CosCurr, const float SinCurr, const float CosNext, const float SinNext)
+		{
+				return FVector2D(CosCurr - CosNext, SinNext - SinCurr) * Bevel;
+			});
+		break;
+	}
+	case EText3DBevelType::Concave:
+		{
+		BevelCurve(HALF_PI, BevelSegments, [Bevel](const float CosCurr, const float SinCurr, const float CosNext, const float SinNext)
+			{
+				return FVector2D(SinNext - SinCurr, CosCurr - CosNext) * Bevel;
+			});
+		break;
+			}
+	case EText3DBevelType::HalfCircle:
+	{
+		BevelCurve(PI, BevelSegments, [Bevel](const float CosCurr, const float SinCurr, const float CosNext, const float SinNext)
+		{
+			return FVector2D(SinCurr - SinNext, CosCurr - CosNext) * Bevel;
+		});
+		break;
+		}
+	case EText3DBevelType::OneStep:
+	{
+		BevelWithSteps(Bevel, 1, BevelSegments);
+		break;
+	}
+	case EText3DBevelType::TwoSteps:
+	{
+		BevelWithSteps(Bevel, 2, BevelSegments);
+		break;
+	}
+	case EText3DBevelType::Engraved:
+	{
+		BevelLinearWithSegments(-Bevel, 0.f, BevelSegments, FVector2D(-1.f, 0.f));
+		BevelLinearWithSegments(0.f, Bevel, BevelSegments, FVector2D(0.f, -1.f));
+		BevelLinearWithSegments(Bevel, 0.f, BevelSegments, FVector2D(1.f, 0.f));
+		break;
+	}
+	default:
+		break;
+	}
+}
+
+void FText3DGlyphMeshBuilder::CreateExtrudeMesh(float Extrude, float Bevel, const EText3DBevelType Type, bool bFlipNormals)
+{
+	if (FMath::IsNearlyZero(Extrude))
+	{
+		return;
+	}
+
+	Bevel = FMath::Max(UE_SMALL_NUMBER, Bevel);
+
+	if (Type != EText3DBevelType::HalfCircle)
+	{
+		Bevel = FMath::Clamp(Bevel, 0.0f, Extrude / 2.f);
+	}
+
+	if (Type != EText3DBevelType::HalfCircle && Type != EText3DBevelType::Engraved)
+	{
+		Extrude -= Bevel * 2.0f;
+	}
+
+	Data->SetCurrentGroup(EText3DGroupType::Extrude, 0.f);
+
+	const FVector2D Normal(1.f, 0.f);
+	Data->PrepareSegment(Extrude, 0.f, Normal, Normal);
+
+	Contours->Reset();
+
+
+	TArray<float> TextureCoordinateVs;
+
+	for (FText3DGlyphContour& Contour : *Contours)
+	{
+		// Compute TexCoord.V-s for each point
+		TextureCoordinateVs.Reset(Contour.Num() - 1);
+		const FText3DGlyphPartPtr First = Contour[0];
+		TextureCoordinateVs.Add(First->Length());
+
+		int32 Index = 1;
+		for (FText3DGlyphPartConstPtr Edge = First->Next; Edge != First->Prev; Edge = Edge->Next)
+		{
+			TextureCoordinateVs.Add(TextureCoordinateVs[Index - 1] + Edge->Length());
+			Index++;
+		}
+
+
+		const float ContourLength = TextureCoordinateVs.Last() + Contour.Last()->Length();
+
+		if (FMath::IsNearlyZero(ContourLength))
+		{
+			continue;
+		}
+
+
+		for (float& PointY : TextureCoordinateVs)
+		{
+			PointY /= ContourLength;
+		}
+
+		// Duplicate contour
+		Data->SetTarget(0.f, 0.f);
+		const bool bFirstSmooth = First->bSmooth;
+		// It's set to sharp because we need 2 vertices with TexCoord.Y values 0 and 1 (for smooth points only one vertex is added)
+		First->bSmooth = false;
+
+		// First point in contour is processed separately
+		{
+			EmptyPaths(First);
+			ExpandPointWithoutAddingVertices(First);
+
+			const FVector2D TexCoordPrev(0.f, 0.f);
+			const FVector2D TexCoordCurr(0.f, 1.f);
+
+			if (bFirstSmooth)
+			{
+				AddVertexSmooth(First, TexCoordPrev);
+				AddVertexSmooth(First, TexCoordCurr);
+			}
+			else
+			{
+				AddVertexSharp(First, First->Prev, TexCoordPrev);
+				AddVertexSharp(First, First, TexCoordCurr);
+			}
+		}
+
+		Index = 1;
+		for (FText3DGlyphPartPtr Point = First->Next; Point != First; Point = Point->Next)
+		{
+			EmptyPaths(Point);
+			ExpandPoint(Point, {0.f, 1.f - TextureCoordinateVs[Index++ - 1]});
+		}
+
+
+		// Add extruded vertices
+		Data->SetTarget(Data->GetPlannedExtrude(), Data->GetPlannedExpand());
+
+		// Similarly to duplicating vertices, first point is processed separately
+		{
+			ExpandPointWithoutAddingVertices(First);
+
+			const FVector2D TexCoordPrev(1.f, 0.f);
+			const FVector2D TexCoordCurr(1.f, 1.f);
+
+			if (bFirstSmooth)
+			{
+				AddVertexSmooth(First, TexCoordPrev);
+				AddVertexSmooth(First, TexCoordCurr);
+			}
+			else
+			{
+				AddVertexSharp(First, First->Prev, TexCoordPrev);
+				AddVertexSharp(First, First, TexCoordCurr);
+			}
+		}
+
+		Index = 1;
+		for (FText3DGlyphPartPtr Point = First->Next; Point != First; Point = Point->Next)
+		{
+			ExpandPoint(Point, {1.f, 1.f - TextureCoordinateVs[Index++ - 1]});
+		}
+
+		for (const FText3DGlyphPartPtr& Edge : Contour)
+		{
+			Data->FillEdge(Edge, false, bFlipNormals);
+		}
+	}
+}
+
+void FText3DGlyphMeshBuilder::MirrorGroup(const EText3DGroupType TypeIn, const EText3DGroupType TypeOut, const float Extrude)
+{
+	TText3DGroupList& Groups = Glyph->GetGroups();
+
+	const FText3DPolygonGroup GroupIn = Groups[static_cast<int32>(TypeIn)];
+	const FText3DPolygonGroup GroupNext = Groups[static_cast<int32>(TypeIn) + 1];
+
+	const int32 VerticesInNum = GroupNext.FirstVertex - GroupIn.FirstVertex;
+	const int32 TrianglesInNum = GroupNext.FirstTriangle - GroupIn.FirstTriangle;
+
+	FMeshDescription& MeshDescription = Glyph->GetMeshDescription();
+	const int32 TotalVerticesNum = MeshDescription.Vertices().Num();
+
+	Data->SetCurrentGroup(TypeOut, 0.f);
+	Data->AddVertices(VerticesInNum);
+
+	FStaticMeshAttributes& StaticMeshAttributes = Glyph->GetStaticMeshAttributes();
+	TVertexAttributesRef<FVector3f> VertexPositions = StaticMeshAttributes.GetVertexPositions();
+	TVertexInstanceAttributesRef<FVector3f> VertexNormals = StaticMeshAttributes.GetVertexInstanceNormals();
+	TVertexInstanceAttributesRef<FVector3f> VertexTangents = StaticMeshAttributes.GetVertexInstanceTangents();
+	TVertexInstanceAttributesRef<FVector2f> VertexUVs = StaticMeshAttributes.GetVertexInstanceUVs();
+
+	for (int32 VertexIndex = 0; VertexIndex < VerticesInNum; VertexIndex++)
+	{
+		const FVertexID VertexID(GroupIn.FirstVertex + VertexIndex);
+		const FVertexInstanceID InstanceID(static_cast<uint32>(VertexID.GetValue()));
+
+		const FVector Position = (FVector)VertexPositions[VertexID];
+		const FVector Normal = (FVector)VertexNormals[InstanceID];
+		const FVector Tangent = (FVector)VertexTangents[InstanceID];
+
+		Data->AddVertex({ Extrude - Position.X, Position.Y, Position.Z }, { -Tangent.X, Tangent.Y, Tangent.Z }, { -Normal.X, Normal.Y, Normal.Z }, FVector2D(VertexUVs[InstanceID]));
+	}
+
+	Data->AddTriangles(TrianglesInNum);
+
+	for (int32 TriangleIndex = 0; TriangleIndex < TrianglesInNum; TriangleIndex++)
+	{
+		const FTriangleID TriangleID = FTriangleID(GroupIn.FirstTriangle + TriangleIndex);
+		TArrayView<const FVertexInstanceID> VertexInstanceIDs = MeshDescription.GetTriangleVertexInstances(TriangleID);
+
+		uint32 Instance0 = static_cast<uint32>(TotalVerticesNum + VertexInstanceIDs[0].GetValue() - GroupIn.FirstVertex);
+		uint32 Instance2 = static_cast<uint32>(TotalVerticesNum + VertexInstanceIDs[2].GetValue() - GroupIn.FirstVertex);
+		uint32 Instance1 = static_cast<uint32>(TotalVerticesNum + VertexInstanceIDs[1].GetValue() - GroupIn.FirstVertex);
+		Data->AddTriangle(Instance0, Instance2, Instance1);
+	}
+}
+
+void FText3DGlyphMeshBuilder::AddToVertexCount(const TText3DGlyphContourNodeShared& Node, int32& OutVertexCount)
+{
+	for (const TText3DGlyphContourNodeShared& Child : Node->Children)
+	{
+		OutVertexCount += Child->Contour->VertexCount();
+		AddToVertexCount(Child, OutVertexCount);
+	}
+}
+
+void FText3DGlyphMeshBuilder::TriangulateAndConvert(const TText3DGlyphContourNodeShared& Node, int32& OutVertexIndex, const bool bOutline)
+{
+	// If this is solid region
+	if (!Node->bClockwise)
+	{
+		int32 VertexCount = 0;
+		UE::Geometry::FConstrainedDelaunay2f Triangulation;
+		Triangulation.FillRule = UE::Geometry::FConstrainedDelaunay2f::EFillRule::Positive;
+
+		const TSharedPtr<FText3DGlyphContourList> ContoursLocal = Contours;
+		const TSharedRef<FText3DGlyphData> DataLocal = Data;
+		auto ProcessContour = [ContoursLocal, DataLocal, &VertexCount, &Triangulation, bOutline](const TText3DGlyphContourNodeShared NodeIn)
+		{
+			// Create contour in old format
+			FText3DGlyphContour& Contour = ContoursLocal->Add();
+			const FPolygon2f& Polygon = *NodeIn->Contour;
+
+			for (const FVector2f& Vertex : Polygon.GetVertices())
+			{
+				// Add point to contour in old format
+				const FText3DGlyphPartPtr Point = MakeShared<FText3DGlyphPart>();
+				Contour.Add(Point);
+				Point->Position = FVector2D(Vertex);
+
+				// Add point to mesh
+				const int32 VertexID = DataLocal->AddVertex(Point->Position, { 1.f, 0.f }, { -1.f, 0.f, 0.f });
+
+				Point->PathPrev.Add(VertexID);
+				Point->PathNext.Add(VertexID);
+			}
+
+			VertexCount += Polygon.VertexCount();
+
+			// Add contour to triangulation
+			if (!bOutline)
+			{
+				Triangulation.Add(Polygon, NodeIn->bClockwise);
+			}
+		};
+
+
+		// Outer
+		ProcessContour(Node);
+
+		// Holes
+		for (const TText3DGlyphContourNodeShared& Child : Node->Children)
+		{
+			ProcessContour(Child);
+		}
+
+		if (!bOutline)
+		{
+			Triangulation.Triangulate();
+			const TArray<FIndex3i>& Triangles = Triangulation.Triangles;
+			Data->AddTriangles(Triangles.Num());
+
+			for (const FIndex3i& Triangle : Triangles)
+			{
+				Data->AddTriangle(OutVertexIndex + Triangle.A, OutVertexIndex + Triangle.C, OutVertexIndex + Triangle.B);
+			}
+		}
+
+		OutVertexIndex += VertexCount;
+	}
+
+	// Continue with children
+	for (const TText3DGlyphContourNodeShared& Child : Node->Children)
+	{
+		TriangulateAndConvert(Child, OutVertexIndex, bOutline);
+	}
+}
+
+void FText3DGlyphMeshBuilder::MakeOutline(float OutlineExpand)
+{
+	FText3DGlyphContourList InitialContours = *Contours;
+
+	for (FText3DGlyphContour& Contour : InitialContours)
+	{
+		Algo::Reverse(Contour);
+
+		for (FText3DGlyphPartPtr& Point : Contour)
+		{
+			Swap(Point->Prev, Point->Next);
+			Point->Normal *= -1.f;
+		}
+
+
+		const FText3DGlyphPartPtr First = Contour[0];
+		const FText3DGlyphPartPtr Last = Contour.Last();
+
+		const FVector2D FirstTangentX = First->TangentX;
+
+		for (FText3DGlyphPartPtr Edge = First; Edge != Last; Edge = Edge->Next)
+		{
+			Edge->TangentX = -Edge->Next->TangentX;
+		}
+
+		Last->TangentX = -FirstTangentX;
+	}
+
+
+	const FVector2D Normal = {0.f, -1.f};
+	BevelLinear(0.f, OutlineExpand, Normal, Normal, false);
+
+
+	Contours->Reset();
+
+
+	TDoubleLinkedList<FText3DGlyphContour>::TDoubleLinkedListNode* Node = InitialContours.GetHead();
+	while (Node)
+	{
+		Contours->AddTail(Node);
+		InitialContours.RemoveNode(Node, false);
+
+		Node = InitialContours.GetHead();
+	}
+}
+
+void FText3DGlyphMeshBuilder::BevelLinearWithSegments(const float Extrude, const float Expand, const int32 BevelSegments, const FVector2D Normal)
+{
+	for (int32 Index = 0; Index < BevelSegments; Index++)
+	{
+		BevelLinear(Extrude / BevelSegments, Expand / BevelSegments, Normal, Normal, false);
+	}
+}
+
+void FText3DGlyphMeshBuilder::BevelCurve(const float Angle, const int32 BevelSegments, TFunction<FVector2D(const float CurrentCos, const float CurrentSin, const float NextCos, const float Next)> ComputeOffset)
+{
+	float CosCurr = 0.0f;
+	float SinCurr = 0.0f;
+
+	float CosNext = 0.0f;
+	float SinNext = 0.0f;
+
+	FVector2D OffsetNext;
+	bool bSmoothNext = false;
+
+	FVector2D NormalNext;
+	FVector2D NormalEnd;
+
+	auto UpdateAngle = [Angle, &CosNext, &SinNext, BevelSegments](const int32 Index)
+	{
+		const float Step = Angle / BevelSegments;
+		FMath::SinCos(&SinNext, &CosNext, Index * Step);
+	};
+
+	auto MakeStep = [UpdateAngle, &OffsetNext, ComputeOffset, &CosCurr, &SinCurr, &CosNext, &SinNext, &NormalNext](int32 Index)
+	{
+		UpdateAngle(Index);
+		OffsetNext = ComputeOffset(CosCurr, SinCurr, CosNext, SinNext);
+		NormalNext = FVector2D(OffsetNext.X, -OffsetNext.Y).GetSafeNormal();
+	};
+
+
+	UpdateAngle(0);
+
+	CosCurr = CosNext;
+	SinCurr = SinNext;
+
+	MakeStep(1);
+	for (int32 Index = 0; Index < BevelSegments; Index++)
+	{
+		CosCurr = CosNext;
+		SinCurr = SinNext;
+
+		const FVector2D OffsetCurr = OffsetNext;
+
+		const FVector2D NormalCurr = NormalNext;
+		FVector2D NormalStart;
+
+		const bool bFirst = (Index == 0);
+		const bool bLast = (Index == BevelSegments - 1);
+
+		const bool bSmooth = bSmoothNext;
+
+		if (!bLast)
+		{
+			MakeStep(Index + 2);
+			bSmoothNext = FVector2D::DotProduct(NormalCurr, NormalNext) >= -FText3DGlyphPart::CosMaxAngleSides;
+		}
+
+		NormalStart = bFirst ? NormalCurr : (bSmooth ? NormalEnd : NormalCurr);
+		NormalEnd = bLast ? NormalCurr : (bSmoothNext ? (NormalCurr + NormalNext).GetSafeNormal() : NormalCurr);
+
+		BevelLinear(OffsetCurr.X, OffsetCurr.Y, NormalStart, NormalEnd, bSmooth);
+	}
+}
+
+void FText3DGlyphMeshBuilder::BevelWithSteps(const float Bevel, const int32 Steps, const int32 BevelSegments)
+{
+	const float BevelPerStep = Bevel / Steps;
+
+	for (int32 Step = 0; Step < Steps; Step++)
+	{
+		BevelLinearWithSegments(BevelPerStep, 0.f, BevelSegments, FVector2D(1.f, 0.f));
+		BevelLinearWithSegments(0.f, BevelPerStep, BevelSegments, FVector2D(0.f, -1.f));
+	}
+}
+
+void FText3DGlyphMeshBuilder::BevelLinear(const float Extrude, const float Expand, FVector2D NormalStart, FVector2D NormalEnd, const bool bSmooth)
+{
+	Data->PrepareSegment(Extrude, Expand, NormalStart, NormalEnd);
+	Contours->Reset();
+
+	if (!bSmooth)
+	{
+		DuplicateContourVertices();
+	}
+
+	BevelPartsWithoutIntersectingNormals();
+
+	Data->IncreaseDoneExtrude();
+}
+
+void FText3DGlyphMeshBuilder::DuplicateContourVertices()
+{
+	Data->SetTarget(0.f, 0.f);
+
+	for (FText3DGlyphContour& Contour : *Contours)
+	{
+		for (const FText3DGlyphPartPtr& Point : Contour)
+		{
+			EmptyPaths(Point);
+			// Duplicate points of contour (expansion with value 0)
+			ExpandPoint(Point);
+		}
+	}
+}
+
+void FText3DGlyphMeshBuilder::BevelPartsWithoutIntersectingNormals()
+{
+	Data->SetTarget(Data->GetPlannedExtrude(), Data->GetPlannedExpand());
+	const float MaxExpand = Data->GetPlannedExpand();
+
+	const bool bFlipNormals = FMath::Sign(Data->GetPlannedExpand()) < 0;
+	for (FText3DGlyphContour& Contour : *Contours)
+	{
+		for (const FText3DGlyphPartPtr& Point : Contour)
+		{
+			if (!FMath::IsNearlyEqual(Point->DoneExpand, MaxExpand) || FMath::IsNearlyZero(MaxExpand))
+			{
+				ExpandPoint(Point);
+			}
+
+			const float Delta = MaxExpand - Point->DoneExpand;
+
+			Point->AvailableExpandNear -= Delta;
+			Point->DecreaseExpandsFar(Delta);
+		}
+
+		for (const FText3DGlyphPartPtr& Edge : Contour)
+		{
+			Data->FillEdge(Edge, false, bFlipNormals);
+		}
+	}
+}
+
+void FText3DGlyphMeshBuilder::EmptyPaths(const FText3DGlyphPartPtr& Point) const
+{
+	Point->PathPrev.Empty();
+	Point->PathNext.Empty();
+}
+
+void FText3DGlyphMeshBuilder::ExpandPoint(const FText3DGlyphPartPtr& Point, const FVector2D TextureCoordinates)
+{
+	ExpandPointWithoutAddingVertices(Point);
+
+	if (Point->bSmooth)
+	{
+		AddVertexSmooth(Point, TextureCoordinates);
+	}
+	else
+	{
+		AddVertexSharp(Point, Point->Prev, TextureCoordinates);
+		AddVertexSharp(Point, Point, TextureCoordinates);
+	}
+}
+
+void FText3DGlyphMeshBuilder::ExpandPointWithoutAddingVertices(const FText3DGlyphPartPtr& Point) const
+{
+	Point->Position = Data->Expanded(Point);
+	const int32 FirstAdded = Data->AddVertices(Point->bSmooth ? 1 : 2);
+
+	Point->PathPrev.Add(FirstAdded);
+	Point->PathNext.Add(Point->bSmooth ? FirstAdded : FirstAdded + 1);
+}
+
+void FText3DGlyphMeshBuilder::AddVertexSmooth(const FText3DGlyphPartConstPtr& Point, const FVector2D TextureCoordinates)
+{
+	const FText3DGlyphPartConstPtr Curr = Point;
+	const FText3DGlyphPartConstPtr Prev = Point->Prev;
+
+	Data->AddVertex(Point, (Prev->TangentX + Curr->TangentX).GetSafeNormal(), (Data->ComputeTangentZ(Prev, Point->DoneExpand) + Data->ComputeTangentZ(Curr, Point->DoneExpand)).GetSafeNormal(), TextureCoordinates);
+}
+
+void FText3DGlyphMeshBuilder::AddVertexSharp(const FText3DGlyphPartConstPtr& Point, const FText3DGlyphPartConstPtr& Edge, const FVector2D TextureCoordinates)
+{
+	Data->AddVertex(Point, Edge->TangentX, Data->ComputeTangentZ(Edge, Point->DoneExpand).GetSafeNormal(), TextureCoordinates);
+}
diff --git a/Engine/Plugins/Experimental/Text3D/Source/Text3D/Private/GeometryBuilders/Text3DGlyphMeshBuilder.h b/Engine/Plugins/Experimental/Text3D/Source/Text3D/Private/GeometryBuilders/Text3DGlyphMeshBuilder.h
new file mode 100644
--- /dev/null
+++ b/Engine/Plugins/Experimental/Text3D/Source/Text3D/Private/GeometryBuilders/Text3DGlyphMeshBuilder.h
@@ -0,0 +1,139 @@
+// Copyright Epic Games, Inc. All Rights Reserved.
+
+#pragma once
+
+#include "GeometryBuilders/Text3DGlyphContourNode.h"
+#include "GeometryBuilders/Text3DGlyphPart.h"
+#include "Text3DTypes.h"
+
+class FText3DGlyphContourList;
+class FText3DGlyphData;
+class FText3DGlyph;
+class UMaterial;
+class UStaticMesh;
+
+/** Creates the glyph mesh from glyph data */
+class FText3DGlyphMeshBuilder final
+{
+public:
+	FText3DGlyphMeshBuilder();
+
+	/**
+	 * CreateMeshes
+	 * @param Root - Tree of contours.	 
+	 * @param Extrude - Orthogonal (to front cap) offset value.
+	 * @param Bevel - Bevel value (bevel happens before extrude).
+	 * @param Type - Defines shape of beveled part.
+	 * @param BevelSegments - Segments count.
+	 * @param bOutline - Front face has outline and is not filled.
+	 * @param OutlineExpand - Offsets the outline by the specified amount.
+	 */
+	void CreateMeshes(const TText3DGlyphContourNodeShared& Root, const float Extrude, const float Bevel, const EText3DBevelType Type, const int32 BevelSegments, const bool bOutline, const float OutlineExpand = 0.5f);
+	void SetFrontAndBevelTextureCoordinates(const float Bevel);
+	void MirrorGroups(const float Extrude);
+	void BuildMesh(UStaticMesh* StaticMesh, UMaterial* DefaultMaterial);
+
+private:
+	TSharedPtr<FText3DGlyph> Glyph;
+	TSharedRef<FText3DGlyphData> Data;
+	TSharedPtr<FText3DGlyphContourList> Contours;
+
+
+	/**
+	 * Create 'Front' part of glyph.
+	 * @param Root - Tree root.
+	 * @param bOutline - Front face has outline and is not filled.
+	 * @param OutlineExpand - Offsets the outline by the specified amount.
+	 */
+	void CreateFrontMesh(const TText3DGlyphContourNodeShared& Root, const bool bOutline, const float& OutlineExpand = 0.5f);
+	/**
+	 * Create 'Bevel' part of glyph (actually half of it, will be mirrored later).
+	 * @param Bevel - Bevel value (bevel happens before extrude).
+	 * @param Type - Defines shape of beveled part.
+	 * @param BevelSegments - Segments count.
+	 */
+	void CreateBevelMesh(const float Bevel, const EText3DBevelType Type, const int32 BevelSegments);
+	/**
+	 * Create 'Extrude' part of glyph.
+	 * @param Extrude - Extrude value.
+	 * @param Bevel - Bevel value (bevel happens before extrude).
+	 * @param Type - Defines shape of beveled part.
+	 * @param bFlipNormals - Reverse the geometry facing direction.
+	 */
+	void CreateExtrudeMesh(float Extrude, float Bevel, const EText3DBevelType Type, bool bFlipNormals = false);
+
+	void MirrorGroup(const EText3DGroupType TypeIn, const EText3DGroupType TypeOut, const float Extrude);
+
+
+	/**
+	 * Recursively compute vertex count in tree.
+	 * @param Node - Tree node.
+	 * @param OutVertexCount - Pointer to counter.
+	 */
+	void AddToVertexCount(const TText3DGlyphContourNodeShared& Node, int32& OutVertexCount);
+	/**
+	 * Triangulate solid region with holes and convert contours to old format.
+	 * @param Node - Tree node.
+	 * @param OutVertexIndex - Pointer to last used vertex index.
+	 * @param bOutline - Front face has outline and is not filled.
+	 */
+	void TriangulateAndConvert(const TText3DGlyphContourNodeShared& Node, int32& OutVertexIndex, const bool bOutline);
+
+	void MakeOutline(float OutlineExpand = 0.5f);
+
+	void BevelLinearWithSegments(const float Extrude, const float Expand, const int32 BevelSegments, const FVector2D Normal);
+	void BevelCurve(const float Angle, const int32 BevelSegments, TFunction<FVector2D(const float CurrentCos, const float CurrentSin, const float NextCos, const float Next)> ComputeOffset);
+	void BevelWithSteps(const float Bevel, const int32 Steps, const int32 BevelSegments);
+
+	/**
+	 * Bevel one segment.
+	 * @param Extrude - Documented in ContourList.h.
+	 * @param Expand - Documented in ContourList.h.
+	 * @param NormalStart - Normal at start of segment (minimum DoneExpand).
+	 * @param NormalEnd - Normal at end of segment (maximum DoneExpand).
+	 * @param bSmooth - Is angle between start of this segment and end of previous segment smooth?
+	 */
+	void BevelLinear(const float Extrude, const float Expand, FVector2D NormalStart, FVector2D NormalEnd, const bool bSmooth);
+
+	/**
+	 * Duplicate contour vertices (used to make sharp angle between bevel steps)
+	 */
+	void DuplicateContourVertices();
+
+	/**
+	 * Continue with trivial bevel till FData::Expand.
+	 */
+	void BevelPartsWithoutIntersectingNormals();
+
+	/**
+	 * Clear PathPrev and PathNext.
+	 * @param Point - Point which paths should be cleared.
+	 */
+	void EmptyPaths(const FText3DGlyphPartPtr& Point) const;
+
+	/**
+	 * Same as previous function but does not cover the intersection-case.
+	 * @param Point - Point that should be expanded.
+	 * @param TextureCoordinates - Texcture coordinates of added vertices.
+	 */
+	void ExpandPoint(const FText3DGlyphPartPtr& Point, const FVector2D TextureCoordinates = FVector2D(0.f, 0.f));
+	/**
+	 * Common code for expanding, vertices are added uninitialized.
+	 * @param Point - Expanded point.
+	 */
+	void ExpandPointWithoutAddingVertices(const FText3DGlyphPartPtr& Point) const;
+
+	/**
+	 * Add vertex for smooth point.
+	 * @param Point - Expanded point.
+	 * @param TextureCoordinates - Texture coordinates of added vertex.
+	 */
+	void AddVertexSmooth(const FText3DGlyphPartConstPtr& Point, const FVector2D TextureCoordinates);
+	/**
+	 * Add vertex for sharp point.
+	 * @param Point - Expanded point.
+	 * @param Edge - Edge from which TangentX and TangentZ will be assigned.
+	 * @param TextureCoordinates - Texture coordinates of added vertex.
+	 */
+	void AddVertexSharp(const FText3DGlyphPartConstPtr& Point, const FText3DGlyphPartConstPtr& Edge, const FVector2D TextureCoordinates);
+};
diff --git a/Engine/Plugins/Experimental/Text3D/Source/Text3D/Private/GeometryBuilders/Text3DGlyphPart.cpp b/Engine/Plugins/Experimental/Text3D/Source/Text3D/Private/GeometryBuilders/Text3DGlyphPart.cpp
new file mode 100644
--- /dev/null
+++ b/Engine/Plugins/Experimental/Text3D/Source/Text3D/Private/GeometryBuilders/Text3DGlyphPart.cpp
@@ -0,0 +1,135 @@
+// Copyright Epic Games, Inc. All Rights Reserved.
+
+#include "GeometryBuilders/Text3DGlyphPart.h"
+
+constexpr float FText3DGlyphPart::CosMaxAngleSideTangent;
+constexpr float FText3DGlyphPart::CosMaxAngleSides;
+
+FText3DGlyphPart::FText3DGlyphPart()
+{
+	Position = FVector2D::ZeroVector;
+	TangentX = FVector2D::ZeroVector;
+	Normal = FVector2D::ZeroVector;
+	InitialPosition = FVector2D::ZeroVector;
+	Prev = nullptr;
+	Next = nullptr;
+	bSmooth = false;
+	AvailableExpandNear = 0.0f;
+	DoneExpand = 0.f;
+}
+
+FText3DGlyphPart::FText3DGlyphPart(const FText3DGlyphPartConstPtr& Other)
+{
+	Position = Other->Position;
+	DoneExpand = Other->DoneExpand;
+	TangentX = Other->TangentX;
+	Normal = Other->Normal;
+	bSmooth = Other->bSmooth;
+	InitialPosition = Other->InitialPosition;
+	PathPrev = Other->PathPrev;
+	PathNext = Other->PathNext;
+	AvailableExpandNear = Other->AvailableExpandNear;
+}
+
+float FText3DGlyphPart::TangentsDotProduct() const
+{
+	check(Prev);
+	return FVector2D::DotProduct(-Prev->TangentX, TangentX);
+}
+
+float FText3DGlyphPart::Length() const
+{
+	check(Next);
+
+	return (Next->Position - Position).Size();
+}
+
+void FText3DGlyphPart::ResetDoneExpand()
+{
+	DoneExpand = 0.f;
+}
+
+void FText3DGlyphPart::ComputeTangentX()
+{
+	check(Next);
+	TangentX = (Next->Position - Position).GetSafeNormal();
+}
+
+bool FText3DGlyphPart::ComputeNormal()
+{
+	check(Prev);
+
+	// Scale is needed to make ((p_(i+1) + k * n_(i+1)) - (p_i + k * n_i)) parallel to (p_(i+1) - p_i). Also (k) is distance between original edge and this edge after expansion with value (k).
+	const float OneMinusADotC = 1.0f - TangentsDotProduct();
+
+	if (FMath::IsNearlyZero(OneMinusADotC))
+	{
+		return false;
+	}
+
+	const FVector2D A = -Prev->TangentX;
+	const FVector2D C = TangentX;
+
+	Normal = A + C;
+	const float NormalLength2 = Normal.SizeSquared();
+	const float Scale = -FPlatformMath::Sqrt(2.0f / OneMinusADotC);
+
+	// If previous and next edge are nearly on one line
+	if (FMath::IsNearlyZero(NormalLength2, 0.0001f))
+	{
+		Normal = FVector2D(A.Y, -A.X) * Scale;
+	}
+	else
+	{
+		// Sign of cross product is needed to be sure that Normal is directed outside.
+		Normal *= -Scale * FPlatformMath::Sign(FVector2D::CrossProduct(A, C)) / FPlatformMath::Sqrt(NormalLength2);
+	}
+
+	Normal.Normalize();
+
+	return true;
+}
+
+void FText3DGlyphPart::ComputeSmooth()
+{
+	bSmooth = TangentsDotProduct() <= CosMaxAngleSides;
+}
+
+bool FText3DGlyphPart::ComputeNormalAndSmooth()
+{
+	if (!ComputeNormal())
+	{
+		return false;
+	}
+
+	ComputeSmooth();
+	return true;
+}
+
+void FText3DGlyphPart::ResetInitialPosition()
+{
+	InitialPosition = Position;
+}
+
+void FText3DGlyphPart::ComputeInitialPosition()
+{
+	InitialPosition = Position - DoneExpand * Normal;
+}
+
+void FText3DGlyphPart::DecreaseExpandsFar(const float Delta)
+{
+	for (auto It = AvailableExpandsFar.CreateIterator(); It; ++It)
+	{
+		It->Value -= Delta;
+
+		if (It->Value < 0.f)
+		{
+			It.RemoveCurrent();
+		}
+	}
+}
+
+FVector2D FText3DGlyphPart::Expanded(const float Value) const
+{
+	return Position + Normal * Value;
+}
diff --git a/Engine/Plugins/Experimental/Text3D/Source/Text3D/Private/GeometryBuilders/Text3DGlyphPart.h b/Engine/Plugins/Experimental/Text3D/Source/Text3D/Private/GeometryBuilders/Text3DGlyphPart.h
new file mode 100644
--- /dev/null
+++ b/Engine/Plugins/Experimental/Text3D/Source/Text3D/Private/GeometryBuilders/Text3DGlyphPart.h
@@ -0,0 +1,95 @@
+// Copyright Epic Games, Inc. All Rights Reserved.
+
+#pragma once
+
+#include "Containers/Map.h"
+#include "HAL/Platform.h"
+#include "Math/Vector.h"
+#include "Math/Vector2D.h"
+#include "Templates/SharedPointer.h"
+
+struct FText3DGlyphPart;
+
+using FText3DGlyphPartPtr = TSharedPtr<FText3DGlyphPart>;
+using FText3DGlyphPartConstPtr = TSharedPtr<const FText3DGlyphPart>;
+
+/** Part represents point and it's next edge (counterclockwise one). */
+struct FText3DGlyphPart final
+{
+	static constexpr float CosMaxAngleSideTangent = 0.995f;
+	static constexpr float CosMaxAngleSides = -0.9f;
+	
+	class FAvailableExpandsFar final : public TMap<FText3DGlyphPartPtr, float>
+	{
+	public:
+		void Add(const FText3DGlyphPartConstPtr Edge, const float Value)
+		{
+			TMap<FText3DGlyphPartPtr, float>::Add(ConstCastSharedPtr<FText3DGlyphPart>(Edge), Value);
+		}
+
+		void Remove(const FText3DGlyphPartConstPtr Edge)
+		{
+			TMap<FText3DGlyphPartPtr, float>::Remove(ConstCastSharedPtr<FText3DGlyphPart>(Edge));
+		}
+	};
+
+	FText3DGlyphPart();
+	FText3DGlyphPart(const FText3DGlyphPartConstPtr& Other);
+
+	/** Previous part. */
+	FText3DGlyphPartPtr Prev;
+
+	/** Next part. */
+	FText3DGlyphPartPtr Next;
+
+	/** Position, is equal to position of last vertex in paths (in coordinate system of glyph). */
+	FVector2D Position;
+
+	/** Offset in surface of front cap that this point already made. */
+	float DoneExpand;
+
+	FVector2D TangentX;
+
+	/** Point normal, a bisector of angle. */
+	FVector2D Normal;
+
+	/** If true, previous and next edges are in one smoothing group. */
+	bool bSmooth;
+
+	FVector2D InitialPosition;
+
+	// Paths along which triangulation is made. Values that are stored are indices of vertices. If bSmooth == true, both paths store one index (for same DoneExpand value). If not, indices are different.
+	/** Path used for triangulation of previous edge. */
+	TArray<int32> PathPrev;
+
+	/** Path used for triangulation of next edge. */
+	TArray<int32> PathNext;
+
+	/** Offset needed for an IntersectionNear to happen. */
+	float AvailableExpandNear;
+
+	/** List of pairs (edge, offset) for IntersectionFar. */
+	FAvailableExpandsFar AvailableExpandsFar;
+
+	float TangentsDotProduct() const;
+	float Length() const;
+
+	void ResetDoneExpand();
+	void ComputeTangentX();
+	bool ComputeNormal();
+	void ComputeSmooth();
+
+	bool ComputeNormalAndSmooth();
+
+	void ResetInitialPosition();
+	void ComputeInitialPosition();
+
+	void DecreaseExpandsFar(const float Delta);
+
+	/**
+	 * Compute position to which point will be expanded.
+	 * @param Value - Offset in surface of front cap.
+	 * @return Computed position.
+	 */
+	FVector2D Expanded(const float Value) const;
+};
diff --git a/Engine/Plugins/Experimental/Text3D/Source/Text3D/Private/Glyph.cpp b/Engine/Plugins/Experimental/Text3D/Source/Text3D/Private/Glyph.cpp
deleted file mode 100644
diff --git a/Engine/Plugins/Experimental/Text3D/Source/Text3D/Private/GlyphLoader.cpp b/Engine/Plugins/Experimental/Text3D/Source/Text3D/Private/GlyphLoader.cpp
deleted file mode 100644
diff --git a/Engine/Plugins/Experimental/Text3D/Source/Text3D/Private/GlyphLoader.h b/Engine/Plugins/Experimental/Text3D/Source/Text3D/Private/GlyphLoader.h
deleted file mode 100644
diff --git a/Engine/Plugins/Experimental/Text3D/Source/Text3D/Private/LayoutBuilders/Text3DLayout.cpp b/Engine/Plugins/Experimental/Text3D/Source/Text3D/Private/LayoutBuilders/Text3DLayout.cpp
new file mode 100644
--- /dev/null
+++ b/Engine/Plugins/Experimental/Text3D/Source/Text3D/Private/LayoutBuilders/Text3DLayout.cpp
@@ -0,0 +1,21 @@
+// Copyright Epic Games, Inc. All Rights Reserved.
+
+#include "LayoutBuilders/Text3DLayout.h"
+
+#include "Brushes/SlateNoResource.h"
+#include "Framework/Text/SlateTextRun.h"
+
+FText3DLayout::FText3DLayout(const FTextBlockStyle& InStyle)
+	: TextStyle(InStyle)
+{
+	static const FSlateBrush EmptyBrush = FSlateNoResource{};
+	TextStyle.SetUnderlineBrush(EmptyBrush);
+	TextStyle.SetStrikeBrush(EmptyBrush);
+}
+
+TSharedRef<IRun> FText3DLayout::CreateDefaultTextRun(
+	const TSharedRef<FString>& NewText,
+	const FTextRange& NewRange) const
+{
+	return FSlateTextRun::Create(FRunInfo(), NewText, TextStyle, NewRange);
+}
\ No newline at end of file
diff --git a/Engine/Plugins/Experimental/Text3D/Source/Text3D/Private/LayoutBuilders/Text3DLayout.h b/Engine/Plugins/Experimental/Text3D/Source/Text3D/Private/LayoutBuilders/Text3DLayout.h
new file mode 100644
--- /dev/null
+++ b/Engine/Plugins/Experimental/Text3D/Source/Text3D/Private/LayoutBuilders/Text3DLayout.h
@@ -0,0 +1,23 @@
+// Copyright Epic Games, Inc. All Rights Reserved.
+
+#pragma once
+
+#include "Framework/Text/TextLayout.h"
+#include "Styling/SlateTypes.h"
+
+/** An implementation of FTextLayout, which discards most Slate widget specific functionality. */
+class FText3DLayout final : public FTextLayout
+{
+public:
+	/** Optionally provide a custom TextBlock style. */
+	explicit FText3DLayout(const FTextBlockStyle& InStyle = FTextBlockStyle::GetDefault());
+
+protected:
+	/** Parameters relevant to text layout. */
+	FTextBlockStyle TextStyle;
+
+	/** Required but unused override. */
+	virtual TSharedRef<IRun> CreateDefaultTextRun(
+		const TSharedRef<FString>& NewText,
+		const FTextRange& NewRange) const override;
+};
\ No newline at end of file
diff --git a/Engine/Plugins/Experimental/Text3D/Source/Text3D/Private/LayoutBuilders/Text3DLayoutShaper.cpp b/Engine/Plugins/Experimental/Text3D/Source/Text3D/Private/LayoutBuilders/Text3DLayoutShaper.cpp
new file mode 100644
--- /dev/null
+++ b/Engine/Plugins/Experimental/Text3D/Source/Text3D/Private/LayoutBuilders/Text3DLayoutShaper.cpp
@@ -0,0 +1,73 @@
+// Copyright Epic Games, Inc. All Rights Reserved.
+
+#include "LayoutBuilders/Text3DLayoutShaper.h"
+
+#include "Brushes/SlateNoResource.h"
+#include "Containers/Array.h"
+#include "Containers/UnrealString.h"
+#include "Fonts/FontCache.h"
+#include "Framework/Application/SlateApplication.h"
+#include "Framework/Text/PlainTextLayoutMarshaller.h"
+#include "Framework/Text/ShapedTextCache.h"
+#include "Framework/Text/SlateTextRun.h"
+#include "Framework/Text/TextLayout.h"
+#include "Internationalization/Text.h"
+#include "Text3DShapedGlyphLine.h"
+
+TSharedPtr<FText3DLayoutShaper> FText3DLayoutShaper::Get()
+{
+	static TSharedPtr<FText3DLayoutShaper> Instance = MakeShared<FText3DLayoutShaper>(FPrivateToken{});
+	return Instance;
+}
+
+void FText3DLayoutShaper::ShapeBidirectionalText(
+	const FTextBlockStyle& Style,
+	const FString& Text,
+	const TSharedPtr<FTextLayout>& TextLayout,
+	const TSharedPtr<ITextLayoutMarshaller>& TextMarshaller,
+	TArray<FShapedGlyphLine>& OutShapedLines)
+{
+	// @todo: Restore when dependency on SlateApplication is removed. Currently this means text meshes won't be created on the server.
+	// @see: UE-211843
+	if (!FSlateApplication::IsInitialized())
+	{
+		return;
+	}
+
+	TextLayout->ClearLines();
+	TextLayout->ClearLineHighlights();
+	TextLayout->ClearRunRenderers();
+	
+	TextMarshaller->SetText(Text, *TextLayout);
+	TextMarshaller->ClearDirty();
+
+	TextLayout->UpdateLayout();
+
+	TArray<FTextLayout::FLineView> LineViews = TextLayout->GetLineViews();
+
+	// @note: mimics FSlateTextLayout::OnPaint
+	for (const FTextLayout::FLineView& Line : LineViews)
+	{
+		FShapedGlyphLine& ShapedLine = OutShapedLines.AddDefaulted_GetRef();
+		for (const TSharedRef<ILayoutBlock>& Block : Line.Blocks)
+		{
+			FString BlockText;
+			Block->GetRun()->AppendTextTo(BlockText);
+
+			TSharedRef<const FString> BlockTextRef = MakeShared<const FString>(BlockText);
+			
+			const FTextRange BlockRange = Block->GetTextRange();
+			const FLayoutBlockTextContext BlockTextContext = Block->GetTextContext();
+
+			FShapedGlyphSequenceRef ShapedGlyphSequence = ShapedTextCacheUtil::GetShapedTextSubSequence(
+				BlockTextContext.ShapedTextCache,
+				FCachedShapedTextKey(Line.Range, 1.0f, BlockTextContext, Style.Font),
+				BlockRange,
+				**BlockTextRef,
+				BlockTextContext.TextDirection
+			);
+
+			ShapedLine.GlyphsToRender.Append(ShapedGlyphSequence->GetGlyphsToRender());
+		}
+	}
+}
diff --git a/Engine/Plugins/Experimental/Text3D/Source/Text3D/Private/LayoutBuilders/Text3DLayoutShaper.h b/Engine/Plugins/Experimental/Text3D/Source/Text3D/Private/LayoutBuilders/Text3DLayoutShaper.h
new file mode 100644
--- /dev/null
+++ b/Engine/Plugins/Experimental/Text3D/Source/Text3D/Private/LayoutBuilders/Text3DLayoutShaper.h
@@ -0,0 +1,45 @@
+// Copyright Epic Games, Inc. All Rights Reserved.
+
+#pragma once
+
+#include "Containers/Array.h"
+#include "Internationalization/Text.h"
+#include "Templates/SharedPointer.h"
+
+class FTextLayout;
+class ITextLayoutMarshaller;
+struct FShapedGlyphLine;
+struct FTextBlockStyle;
+
+/** A singleton that handles shaping operations and writes the result to a provided TextLayout. */
+class FText3DLayoutShaper final
+{
+public:
+	/** Returns TextShaper singleton. */
+	static TSharedPtr<FText3DLayoutShaper> Get();
+
+	/**
+	 * Arranges the provided text to match the requested layout, accounting for scale, offsets etc.
+	 * Analogous to FSlateFontCache::ShapeBidirectionalText.
+	 */
+	void ShapeBidirectionalText(
+		const FTextBlockStyle& Style,
+		const FString& Text,
+		const TSharedPtr<FTextLayout>& TextLayout,
+		const TSharedPtr<ITextLayoutMarshaller>& TextMarshaller,
+		TArray<FShapedGlyphLine>& OutShapedLines);
+
+private:
+	struct FPrivateToken { explicit FPrivateToken() = default; };
+
+public:
+	explicit FText3DLayoutShaper(FPrivateToken)
+	{
+	}
+
+private:
+	friend class UText3DComponent;
+
+	FText3DLayoutShaper(const FText3DLayoutShaper&) = delete;
+	FText3DLayoutShaper& operator=(const FText3DLayoutShaper&) = delete;
+};
diff --git a/Engine/Plugins/Experimental/Text3D/Source/Text3D/Private/LayoutBuilders/Text3DShapedGlyphLine.cpp b/Engine/Plugins/Experimental/Text3D/Source/Text3D/Private/LayoutBuilders/Text3DShapedGlyphLine.cpp
new file mode 100644
--- /dev/null
+++ b/Engine/Plugins/Experimental/Text3D/Source/Text3D/Private/LayoutBuilders/Text3DShapedGlyphLine.cpp
@@ -0,0 +1,33 @@
+// Copyright Epic Games, Inc. All Rights Reserved.
+
+#include "LayoutBuilders/Text3DShapedGlyphLine.h"
+
+float FShapedGlyphLine::GetAdvance(const int32 InIndex, const float InKerning, const float InWordSpacing) const
+{
+	check(InIndex >= 0 && InIndex < GlyphsToRender.Num());
+
+	const FShapedGlyphEntry& Glyph = GlyphsToRender[InIndex];
+	float Advance = Glyph.XOffset + Glyph.XAdvance;
+
+	if (InIndex < GlyphsToRender.Num() - 1)
+	{
+		// @note: as per FSlateElementBatcher::BuildShapedTextSequence, per Glyph Kerning isn't used
+		Advance += InKerning;
+
+		if (!Glyph.bIsVisible)
+		{
+			Advance += InWordSpacing;
+		}
+	}
+
+	return Advance;
+}
+
+void FShapedGlyphLine::CalculateWidth(const float InKerning, const float InWordSpacing)
+{
+	Width = 0.0f;
+	for (int32 Index = 0; Index < GlyphsToRender.Num(); Index++)
+	{
+		Width += GetAdvance(Index, InKerning, InWordSpacing);
+	}
+}
diff --git a/Engine/Plugins/Experimental/Text3D/Source/Text3D/Private/LayoutBuilders/Text3DShapedGlyphLine.h b/Engine/Plugins/Experimental/Text3D/Source/Text3D/Private/LayoutBuilders/Text3DShapedGlyphLine.h
new file mode 100644
--- /dev/null
+++ b/Engine/Plugins/Experimental/Text3D/Source/Text3D/Private/LayoutBuilders/Text3DShapedGlyphLine.h
@@ -0,0 +1,22 @@
+// Copyright Epic Games, Inc. All Rights Reserved.
+
+#pragma once
+
+#include "Containers/Array.h"
+#include "Fonts/FontCache.h"
+
+/** Contains text line metric with sufficient information to fetch and transform each character. */
+struct FShapedGlyphLine
+{
+	/** Get's the offset from the previous character, accounting for kerning and word spacing. */
+	float GetAdvance(const int32 InIndex, const float InKerning, const float InWordSpacing) const;
+
+	/** Calculates the total width of this line. */
+	void CalculateWidth(const float InKerning, const float InWordSpacing);
+
+	/** The corresponding shaped glyph for each character in this line of text. */
+	TArray<FShapedGlyphEntry> GlyphsToRender;
+
+	/** Stored result of CalculateWidth. */
+	float Width = 0.0f;
+};
\ No newline at end of file
diff --git a/Engine/Plugins/Experimental/Text3D/Source/Text3D/Private/LayoutBuilders/Text3DShapedGlyphText.cpp b/Engine/Plugins/Experimental/Text3D/Source/Text3D/Private/LayoutBuilders/Text3DShapedGlyphText.cpp
new file mode 100644
--- /dev/null
+++ b/Engine/Plugins/Experimental/Text3D/Source/Text3D/Private/LayoutBuilders/Text3DShapedGlyphText.cpp
@@ -0,0 +1,103 @@
+// Copyright Epic Games, Inc. All Rights Reserved.
+
+#include "LayoutBuilders/Text3DShapedGlyphText.h"
+
+#include "Misc/EnumerateRange.h"
+
+FText3DShapedGlyphText::FText3DShapedGlyphText()
+{
+	Reset();
+}
+
+void FText3DShapedGlyphText::Reset()
+{
+	LineHeight = 0.0f;
+	FontAscender = 0.0f;
+	FontDescender = 0.0f;
+	Kerning = 0.0f;
+	WordSpacing = 0.0f;
+	bWrap = false;
+	Lines.Reset();
+	GlyphIndexToFontFace.Empty();
+}
+
+void FText3DShapedGlyphText::CalculateWidth()
+{
+	TArray<FShapedGlyphLine> NewLines;
+	NewLines.Reserve(Lines.Num());
+
+	int32 TextGlyphIdx = 0;
+	for (const FShapedGlyphLine& GlyphLine : Lines)
+	{
+		FShapedGlyphLine* CurrentLine = &NewLines.AddDefaulted_GetRef();
+		TArray<FShapedGlyphEntry> CurrentWord;
+
+		float LineWidth = 0.0f;
+		float CurrentWordLength = 0.0f;
+
+		for (TConstEnumerateRef<FShapedGlyphEntry> CurrentGlyph : EnumerateRange(GlyphLine.GlyphsToRender))
+		{
+			// Trim for proper positioning
+			if (!CurrentGlyph->bIsVisible && CurrentGlyph.GetIndex() == 0)
+			{
+				continue;
+			}
+
+			const bool bWordBreak = !CurrentGlyph->bIsVisible || CurrentGlyph.GetIndex() == GlyphLine.GlyphsToRender.Num() - 1;
+			const float GlyphAdv = GlyphLine.GetAdvance(CurrentGlyph.GetIndex(), Kerning, WordSpacing);
+
+			// If we're at the end the line or at whitespace
+			if (bWrap                               // when we're wrapping
+				&& bWordBreak                       // and at a word break
+				&& LineWidth > MaxWidth             // and the current line is longer than the max
+				&& CurrentWordLength != LineWidth)  // and the line is not just a single word that we can't break
+			{
+				CurrentLine->Width = LineWidth - CurrentWordLength;
+				CurrentLine = &NewLines.AddDefaulted_GetRef();
+				LineWidth = CurrentWordLength;
+			}
+
+			CurrentWord.Add(*CurrentGlyph);
+			LineWidth += GlyphAdv;
+			CurrentWordLength += GlyphAdv;
+
+			if (bWordBreak)
+			{
+				CurrentLine->GlyphsToRender.Append(CurrentWord);
+				CurrentWordLength = 0.0f;
+				CurrentWord.Empty();
+			}
+
+			TextGlyphIdx++;
+		}
+
+		CurrentLine->Width = LineWidth;
+	}
+
+	TextGlyphIdx = 0;
+	for (FShapedGlyphLine& NewLine : NewLines)
+	{
+		if (NewLine.GlyphsToRender.IsEmpty())
+		{
+			continue;
+		}
+
+		const FShapedGlyphEntry& FirstGlyph = NewLine.GlyphsToRender[0];
+		if (!FirstGlyph.bIsVisible)
+		{
+			NewLine.Width -= NewLine.GetAdvance(0, Kerning, WordSpacing);
+			NewLine.GlyphsToRender.RemoveAt(0);
+		}
+
+		const int32 LastIndex = NewLine.GlyphsToRender.Num() - 1;
+		TextGlyphIdx += LastIndex;
+		const FShapedGlyphEntry& LastGlyph = NewLine.GlyphsToRender[LastIndex];
+		if (!LastGlyph.bIsVisible)
+		{
+			NewLine.Width -= NewLine.GetAdvance(LastIndex, Kerning, WordSpacing);
+			NewLine.GlyphsToRender.RemoveAt(LastIndex);
+		}
+	}
+	
+	Lines = NewLines;
+}
diff --git a/Engine/Plugins/Experimental/Text3D/Source/Text3D/Private/LayoutBuilders/Text3DShapedGlyphText.h b/Engine/Plugins/Experimental/Text3D/Source/Text3D/Private/LayoutBuilders/Text3DShapedGlyphText.h
new file mode 100644
--- /dev/null
+++ b/Engine/Plugins/Experimental/Text3D/Source/Text3D/Private/LayoutBuilders/Text3DShapedGlyphText.h
@@ -0,0 +1,28 @@
+// Copyright Epic Games, Inc. All Rights Reserved.
+
+#pragma once
+
+#include "Containers/Array.h"
+#include "Containers/Map.h"
+#include "Text3DShapedGlyphLine.h"
+
+class FFreeTypeFace;
+
+/** Contains text metrics with sufficient information to fetch and transform each line. */
+struct FText3DShapedGlyphText
+{
+	FText3DShapedGlyphText();
+
+	void Reset();
+	void CalculateWidth();
+
+	float LineHeight;
+	float FontAscender;
+	float FontDescender;
+	float Kerning;
+	float WordSpacing;
+	float MaxWidth;
+	bool bWrap;
+	TArray<FShapedGlyphLine> Lines;
+	TMap<uint32, const FFreeTypeFace*> GlyphIndexToFontFace;
+};
diff --git a/Engine/Plugins/Experimental/Text3D/Source/Text3D/Private/Mesh.cpp b/Engine/Plugins/Experimental/Text3D/Source/Text3D/Private/Mesh.cpp
deleted file mode 100644
diff --git a/Engine/Plugins/Experimental/Text3D/Source/Text3D/Private/MeshCreator.cpp b/Engine/Plugins/Experimental/Text3D/Source/Text3D/Private/MeshCreator.cpp
deleted file mode 100644
diff --git a/Engine/Plugins/Experimental/Text3D/Source/Text3D/Private/MeshCreator.h b/Engine/Plugins/Experimental/Text3D/Source/Text3D/Private/MeshCreator.h
deleted file mode 100644
diff --git a/Engine/Plugins/Experimental/Text3D/Source/Text3D/Private/Part.cpp b/Engine/Plugins/Experimental/Text3D/Source/Text3D/Private/Part.cpp
deleted file mode 100644
diff --git a/Engine/Plugins/Experimental/Text3D/Source/Text3D/Private/Part.h b/Engine/Plugins/Experimental/Text3D/Source/Text3D/Private/Part.h
deleted file mode 100644
diff --git a/Engine/Plugins/Experimental/Text3D/Source/Text3D/Private/Renderers/StaticMeshes/Text3DStaticMeshesRenderer.cpp b/Engine/Plugins/Experimental/Text3D/Source/Text3D/Private/Renderers/StaticMeshes/Text3DStaticMeshesRenderer.cpp
new file mode 100644
--- /dev/null
+++ b/Engine/Plugins/Experimental/Text3D/Source/Text3D/Private/Renderers/StaticMeshes/Text3DStaticMeshesRenderer.cpp
@@ -0,0 +1,342 @@
+// Copyright Epic Games, Inc. All Rights Reserved.
+
+#include "Renderers/StaticMeshes/Text3DStaticMeshesRenderer.h"
+
+#include "Algo/Accumulate.h"
+#include "Algo/Count.h"
+#include "Components/SceneComponent.h"
+#include "Components/StaticMeshComponent.h"
+#include "Engine/StaticMesh.h"
+#include "LayoutBuilders/Text3DShapedGlyphText.h"
+#include "LayoutBuilders/Text3DLayoutShaper.h"
+#include "Subsystems/Text3DEngineSubsystem.h"
+#include "Text3DComponent.h"
+#include "Text3DInternalTypes.h"
+#include "Misc/EnumerateRange.h"
+
+void UText3DStaticMeshesRenderer::OnCreate()
+{
+	if (!TextRoot)
+	{
+		UText3DComponent* TextComponent = GetText3DComponent();
+		TextRoot = NewObject<USceneComponent>(TextComponent, TEXT("Text3DRoot"), RF_Transient);
+		TextRoot->SetupAttachment(TextComponent);
+	}
+}
+
+void UText3DStaticMeshesRenderer::OnUpdate(EText3DRendererFlags InFlags)
+{
+	UText3DComponent* TextComponent = GetText3DComponent();
+
+	if (EnumHasAnyFlags(InFlags, EText3DRendererFlags::Geometry))
+	{
+		TextRoot->SetRelativeScale3D(TextComponent->GetTextScale());
+		
+		const TSharedPtr<FText3DShapedGlyphText> ShapedText = TextComponent->GetShapedText();
+
+		// Pre-allocate, avoid new'ing!
+		AllocateGlyphs(Algo::TransformAccumulate(ShapedText->Lines, [](const FShapedGlyphLine& ShapedLine)
+		{
+			return Algo::CountIf(ShapedLine.GlyphsToRender, [](const FShapedGlyphEntry& ShapedGlyph)
+			{
+				return ShapedGlyph.bIsVisible;
+			});
+		},
+		0));
+
+		UText3DEngineSubsystem* Subsystem = UText3DEngineSubsystem::Get();
+
+		const uint32 TypefaceIndex = TextComponent->GetTypeFaceIndex();
+		FCachedFontData& CachedFontData = Subsystem->GetCachedFontData(TextComponent->GetFont(), TypefaceIndex);
+		FGlyphMeshParameters GlyphMeshParameters{TextComponent->GetExtrude(), TextComponent->GetBevel(), TextComponent->GetBevelType(), TextComponent->GetBevelSegments(), TextComponent->HasOutline(), TextComponent->GetOutlineExpand(), TypefaceIndex};
+
+		int32 GlyphIndex = 0;
+		for (TConstEnumerateRef<FShapedGlyphLine> ShapedLine : EnumerateRange(ShapedText->Lines))
+		{
+			FVector LineLocation = TextComponent->GetLineLocation(ShapedLine.GetIndex());
+
+			for (TConstEnumerateRef<FShapedGlyphEntry> ShapedGlyph : EnumerateRange(ShapedLine->GlyphsToRender))
+			{
+				FVector GlyphLocation = LineLocation;
+				LineLocation.Y += ShapedLine->GetAdvance(ShapedGlyph.GetIndex(), ShapedText->Kerning, ShapedText->WordSpacing);
+
+				if (!ShapedGlyph->bIsVisible)
+				{
+					continue;
+				}
+
+				// Count even when mesh is nullptr (allocation still creates components to avoid mesh building in allocation step)
+				const int32 GlyphId = GlyphIndex++;
+
+				const FFreeTypeFace* FontFace = ShapedText->GlyphIndexToFontFace[ShapedGlyph->GlyphIndex];
+				UStaticMesh* CachedMesh = CachedFontData.GetGlyphMesh(ShapedGlyph->GlyphIndex, GlyphMeshParameters, FontFace);
+				if (!CachedMesh || FMath::IsNearlyZero(CachedMesh->GetBounds().SphereRadius))
+				{
+					continue;
+				}
+
+				if (CharacterMeshes.IsValidIndex(GlyphId))
+				{
+					UStaticMeshComponent* StaticMeshComponent = CharacterMeshes[GlyphId];
+					StaticMeshComponent->SetStaticMesh(CachedMesh);
+					StaticMeshComponent->SetVisibility(TextComponent->GetVisibleFlag());
+					StaticMeshComponent->SetHiddenInGame(TextComponent->bHiddenInGame);
+					StaticMeshComponent->SetCastShadow(TextComponent->CastsShadow());
+				}
+				else
+				{
+					// @note: This shouldn't occur, but it does under unknown circumstances (UE-164789) so it should be handled
+					UE_LOG(LogText3D, Error, TEXT("CharacterMesh not found at index %d"), GlyphId);
+				}
+
+				if (CharacterKernings.IsValidIndex(GlyphId))
+				{
+					FTransform Transform;
+					Transform.SetLocation(GlyphLocation);
+					USceneComponent* CharacterKerningComponent = CharacterKernings[GlyphId];
+					CharacterKerningComponent->SetRelativeTransform(Transform);
+				}
+				else
+				{
+					// @note: This shouldn't occur, but it does under unknown circumstances (UE-164789) so it should be handled
+					UE_LOG(LogText3D, Error, TEXT("CharacterKerning not found at index %d"), GlyphId);
+				}
+			}
+		}
+	}
+
+	if (EnumHasAnyFlags(InFlags, EText3DRendererFlags::Layout))
+	{
+		TextRoot->SetRelativeScale3D(TextComponent->GetTextScale());
+
+		const TSharedPtr<FText3DShapedGlyphText> ShapedText = TextComponent->GetShapedText();
+
+		int32 GlyphIndex = 0;
+		for (TConstEnumerateRef<FShapedGlyphLine> ShapedLine : EnumerateRange(ShapedText->Lines))
+		{
+			FVector Location = TextComponent->GetLineLocation(ShapedLine.GetIndex());
+
+			for (TConstEnumerateRef<FShapedGlyphEntry> ShapedGlyph : EnumerateRange(ShapedLine->GlyphsToRender))
+			{
+				const FVector CharLocation = Location;
+				Location.Y += ShapedLine->GetAdvance(ShapedGlyph.GetIndex(), ShapedText->Kerning, ShapedText->WordSpacing);
+
+				if (!ShapedGlyph->bIsVisible)
+				{
+					continue;
+				}
+
+				if (USceneComponent* GlyphKerningComponent = GetGlyphKerningComponent(GlyphIndex))
+				{
+					GlyphKerningComponent->SetRelativeLocation(CharLocation);
+				}
+
+				GlyphIndex++;
+			}
+		}
+	}
+
+	if (EnumHasAnyFlags(InFlags, EText3DRendererFlags::Visibility))
+	{
+		for (UStaticMeshComponent* StaticMeshComponent : CharacterMeshes)
+		{
+			StaticMeshComponent->SetHiddenInGame(TextComponent->bHiddenInGame);
+			StaticMeshComponent->SetVisibility(TextComponent->GetVisibleFlag());
+			StaticMeshComponent->SetCastShadow(TextComponent->CastsShadow());
+		}
+	}
+
+	if (EnumHasAnyFlags(InFlags, EText3DRendererFlags::Material))
+	{
+		using namespace UE::Text3D::Materials;
+
+		for (UStaticMeshComponent* StaticMeshComponent : CharacterMeshes)
+		{
+			for (int32 GroupIndex = 0; GroupIndex < static_cast<int32>(EText3DGroupType::TypeCount); GroupIndex++)
+			{
+				const int32 MaterialIndex = StaticMeshComponent->GetMaterialIndex(SlotNames[GroupIndex]);
+
+				if (MaterialIndex == INDEX_NONE)
+				{
+					continue;
+				}
+
+				UMaterialInterface* Material = TextComponent->GetMaterial(static_cast<EText3DGroupType>(GroupIndex));
+
+				if (Material != StaticMeshComponent->GetMaterial(MaterialIndex))
+				{
+					StaticMeshComponent->SetMaterial(MaterialIndex, Material);
+				}
+			}
+		}
+	}
+}
+
+void UText3DStaticMeshesRenderer::OnClear()
+{
+	for (UStaticMeshComponent* MeshComponent : CharacterMeshes)
+	{
+		if (MeshComponent)
+		{
+			MeshComponent->DetachFromComponent(FDetachmentTransformRules::KeepRelativeTransform);
+			MeshComponent->SetStaticMesh(nullptr);
+			MeshComponent->DestroyComponent();
+		}
+	}
+	CharacterMeshes.Reset();
+
+	for (USceneComponent* KerningComponent : CharacterKernings)
+	{
+		if (KerningComponent)
+		{
+			KerningComponent->DetachFromComponent(FDetachmentTransformRules::KeepRelativeTransform);
+			KerningComponent->DestroyComponent();
+		}
+	}
+	CharacterKernings.Reset();
+
+	if (TextRoot)
+	{
+		TArray<USceneComponent*> ChildComponents;
+		constexpr bool bIncludeChildDescendants = true;
+		TextRoot->GetChildrenComponents(bIncludeChildDescendants, ChildComponents);
+
+		for (USceneComponent* ChildComponent : ChildComponents)
+		{
+			if (IsValid(ChildComponent))
+			{
+				ChildComponent->DetachFromComponent(FDetachmentTransformRules::KeepRelativeTransform);
+				ChildComponent->DestroyComponent();
+			}
+		}
+	}
+}
+
+void UText3DStaticMeshesRenderer::OnDestroy()
+{
+	if (TextRoot)
+	{
+		TextRoot->DestroyComponent();
+		TextRoot = nullptr;
+	}
+}
+
+FName UText3DStaticMeshesRenderer::GetName() const
+{
+	static const FName Name(TEXT("StaticMeshesRenderer"));
+	return Name;
+}
+
+FBox UText3DStaticMeshesRenderer::OnCalculateBounds() const
+{
+	FBox Box(ForceInit);
+
+	for (const UStaticMeshComponent* StaticMeshComponent : CharacterMeshes)
+	{
+		Box += StaticMeshComponent->Bounds.GetBox();
+	}
+
+	return Box;
+}
+
+bool UText3DStaticMeshesRenderer::AllocateGlyphs(int32 Num)
+{
+	int32 DeltaNum = Num - CharacterMeshes.Num();
+	if (DeltaNum == 0)
+	{
+		return false;
+	}
+
+	UText3DComponent* TextComponent = GetText3DComponent();
+
+	// Add characters
+	if (FMath::Sign(DeltaNum) > 0)
+	{
+		int32 GlyphId = CharacterMeshes.Num() - 1;
+		for(int32 CharacterIndex = 0; CharacterIndex < DeltaNum; ++CharacterIndex)
+		{
+			GlyphId++;
+
+			const FName CharacterKerningComponentName = MakeUniqueObjectName(this, USceneComponent::StaticClass(), FName(*FString::Printf(TEXT("CharacterKerning%d"), GlyphId)));
+			USceneComponent* CharacterKerningComponent = NewObject<USceneComponent>(this, CharacterKerningComponentName, RF_Transient);
+
+			CharacterKerningComponent->AttachToComponent(TextRoot, FAttachmentTransformRules::KeepRelativeTransform);
+			CharacterKerningComponent->RegisterComponent();
+			CharacterKernings.Add(CharacterKerningComponent);
+
+			const FName StaticMeshComponentName = MakeUniqueObjectName(this, UStaticMeshComponent::StaticClass(), FName(*FString::Printf(TEXT("StaticMeshComponent%d"), GlyphId)));
+			UStaticMeshComponent* StaticMeshComponent = NewObject<UStaticMeshComponent>(this, StaticMeshComponentName, RF_Transient);
+			StaticMeshComponent->RegisterComponent();
+			StaticMeshComponent->SetVisibility(TextComponent->GetVisibleFlag());
+			StaticMeshComponent->SetHiddenInGame(TextComponent->bHiddenInGame);
+			StaticMeshComponent->SetCastShadow(TextComponent->CastsShadow());
+			CharacterMeshes.Add(StaticMeshComponent);
+
+			StaticMeshComponent->AttachToComponent(CharacterKerningComponent, FAttachmentTransformRules::KeepRelativeTransform);
+		}
+	}
+	// Remove characters
+	else
+	{
+		DeltaNum = FMath::Abs(DeltaNum);
+		for(int32 CharacterIndex = CharacterKernings.Num() - 1 - DeltaNum; CharacterIndex < CharacterKernings.Num(); ++CharacterIndex)
+		{
+			USceneComponent* CharacterKerningComponent = CharacterKernings[CharacterIndex];
+			// If called in quick succession, may already be pending destruction
+			if (IsValid(CharacterKerningComponent))
+			{
+				CharacterKerningComponent->DetachFromComponent(FDetachmentTransformRules::KeepRelativeTransform);
+				CharacterKerningComponent->UnregisterComponent();
+				CharacterKerningComponent->DestroyComponent();
+			}
+
+			UStaticMeshComponent* StaticMeshComponent = CharacterMeshes[CharacterIndex];
+			if (IsValid(StaticMeshComponent))
+			{
+				StaticMeshComponent->DetachFromComponent(FDetachmentTransformRules::KeepRelativeTransform);
+				StaticMeshComponent->UnregisterComponent();
+				StaticMeshComponent->DestroyComponent();
+			}
+		}
+
+		CharacterKernings.RemoveAt(CharacterKernings.Num() - 1 - DeltaNum, DeltaNum);
+		CharacterMeshes.RemoveAt(CharacterMeshes.Num() - 1 - DeltaNum, DeltaNum);
+	}
+
+	return true;
+}
+
+int32 UText3DStaticMeshesRenderer::GetGlyphCount()
+{
+	return TextRoot->GetNumChildrenComponents();
+}
+
+USceneComponent* UText3DStaticMeshesRenderer::GetGlyphKerningComponent(int32 Index)
+{
+	if (!CharacterKernings.IsValidIndex(Index))
+	{
+		return nullptr;
+	}
+
+	return CharacterKernings[Index];
+}
+
+const TArray<USceneComponent*>& UText3DStaticMeshesRenderer::GetGlyphKerningComponents()
+{
+	return CharacterKernings;
+}
+
+UStaticMeshComponent* UText3DStaticMeshesRenderer::GetGlyphMeshComponent(int32 Index)
+{
+	if (!CharacterKernings.IsValidIndex(Index))
+	{
+		return nullptr;
+	}
+
+	return CharacterMeshes[Index];
+}
+
+const TArray<UStaticMeshComponent*>& UText3DStaticMeshesRenderer::GetGlyphMeshComponents()
+{
+	return CharacterMeshes;
+}
diff --git a/Engine/Plugins/Experimental/Text3D/Source/Text3D/Private/Renderers/StaticMeshes/Text3DStaticMeshesRenderer.h b/Engine/Plugins/Experimental/Text3D/Source/Text3D/Private/Renderers/StaticMeshes/Text3DStaticMeshesRenderer.h
new file mode 100644
--- /dev/null
+++ b/Engine/Plugins/Experimental/Text3D/Source/Text3D/Private/Renderers/StaticMeshes/Text3DStaticMeshesRenderer.h
@@ -0,0 +1,75 @@
+// Copyright Epic Games, Inc. All Rights Reserved.
+
+#pragma once
+
+#include "Containers/Array.h"
+#include "Renderers/Text3DRendererBase.h"
+#include "UObject/ObjectPtr.h"
+#include "Text3DStaticMeshesRenderer.generated.h"
+
+class USceneComponent;
+class UStaticMeshComponent;
+class UText3DComponent;
+
+/**
+ * Legacy/default renderer for Text3D
+ * Each text character is rendered as a StaticMesh within its own StaticMeshComponent
+ * Kerning is done through a scene component containing the current character
+ * Text3DComponent
+ * - Text3DRoot (Root)
+ * -- Text3DSceneComponent (Kerning)
+ * --- Text3DStaticMeshComponent (Character)
+ *
+ * Eg: The text "Hello" will be rendered using 5 SceneComponents and 5 StaticMeshComponents
+ */
+UCLASS(DisplayName="StaticMeshesRenderer", ClassGroup="Text3D")
+class UText3DStaticMeshesRenderer : public UText3DRendererBase
+{
+	GENERATED_BODY()
+
+public:
+	/** Gets the number of font glyphs that are currently used */
+	UFUNCTION(BlueprintCallable, Category = "Rendering|Components|Text3D")
+	int32 GetGlyphCount();
+
+	/** Gets the USceneComponent that a glyph is attached to */
+	UFUNCTION(BlueprintCallable, Category = "Rendering|Components|Text3D")
+	USceneComponent* GetGlyphKerningComponent(int32 Index);
+
+	/** Gets all the glyph kerning components */
+	UFUNCTION(BlueprintCallable, Category = "Rendering|Components|Text3D")
+	const TArray<USceneComponent*>& GetGlyphKerningComponents();
+
+	/** Gets the StaticMeshComponent of a glyph */
+	UFUNCTION(BlueprintCallable, Category = "Rendering|Components|Text3D")
+	UStaticMeshComponent* GetGlyphMeshComponent(int32 Index);
+
+	/** Gets all the glyph meshes */
+	UFUNCTION(BlueprintCallable, Category = "Rendering|Components|Text3D")
+	const TArray<UStaticMeshComponent*>& GetGlyphMeshComponents();
+
+protected:
+	//~ Begin UText3DRendererBase
+	virtual void OnCreate() override;
+	virtual void OnUpdate(EText3DRendererFlags InFlags) override;
+	virtual void OnClear() override;
+	virtual void OnDestroy() override;
+	virtual FName GetName() const override;
+	virtual FBox OnCalculateBounds() const override;
+	//~ End UText3DRendererBase
+
+	/** Allocates, or shrinks existing components to match the input number. Returns false if nothing modified. */
+	bool AllocateGlyphs(int32 Num);
+
+	/** Holds all character components */
+	UPROPERTY(Transient)
+	TObjectPtr<USceneComponent> TextRoot;
+
+	/** Each character kerning is held in these components */
+	UPROPERTY(Transient)
+	TArray<TObjectPtr<USceneComponent>> CharacterKernings;
+
+	/** Each character mesh is held in these components */
+	UPROPERTY(Transient)
+	TArray<TObjectPtr<UStaticMeshComponent>> CharacterMeshes;
+};
\ No newline at end of file
diff --git a/Engine/Plugins/Experimental/Text3D/Source/Text3D/Private/Renderers/Text3DRendererBase.cpp b/Engine/Plugins/Experimental/Text3D/Source/Text3D/Private/Renderers/Text3DRendererBase.cpp
new file mode 100644
--- /dev/null
+++ b/Engine/Plugins/Experimental/Text3D/Source/Text3D/Private/Renderers/Text3DRendererBase.cpp
@@ -0,0 +1,70 @@
+// Copyright Epic Games, Inc. All Rights Reserved.
+
+#include "Renderers/Text3DRendererBase.h"
+
+#include "Text3DComponent.h"
+
+UText3DComponent* UText3DRendererBase::GetText3DComponent() const
+{
+	return GetTypedOuter<UText3DComponent>();
+}
+
+void UText3DRendererBase::Create()
+{
+	if (bInitialized)
+	{
+		return;
+	}
+
+	OnCreate();
+	bInitialized = true;
+}
+
+void UText3DRendererBase::Update(EText3DRendererFlags InFlags)
+{
+	if (!bInitialized)
+	{
+		return;
+	}
+
+	OnUpdate(InFlags);
+	
+	if (EnumHasAnyFlags(InFlags, EText3DRendererFlags::Geometry)
+		|| EnumHasAnyFlags(InFlags, EText3DRendererFlags::Layout))
+	{
+		CachedBounds = OnCalculateBounds();
+	}
+}
+
+void UText3DRendererBase::Clear()
+{
+	if (!bInitialized)
+	{
+		return;
+	}
+
+	OnClear();
+	CachedBounds.Reset();
+}
+
+void UText3DRendererBase::Destroy()
+{
+	if (!bInitialized)
+	{
+		return;
+	}
+
+	OnDestroy();
+	CachedBounds.Reset();
+	bInitialized = false;
+}
+
+FBox UText3DRendererBase::GetBounds() const
+{
+	if (!CachedBounds.IsSet())
+	{
+		return FBox(ForceInitToZero);
+	}
+
+	return CachedBounds.GetValue();
+}
diff --git a/Engine/Plugins/Experimental/Text3D/Source/Text3D/Private/Subsystems/Text3DEngineSubsystem.cpp b/Engine/Plugins/Experimental/Text3D/Source/Text3D/Private/Subsystems/Text3DEngineSubsystem.cpp
new file mode 100644
--- /dev/null
+++ b/Engine/Plugins/Experimental/Text3D/Source/Text3D/Private/Subsystems/Text3DEngineSubsystem.cpp
@@ -0,0 +1,448 @@
+// Copyright Epic Games, Inc. All Rights Reserved.
+
+#include "Text3DEngineSubsystem.h"
+#include "Containers/Ticker.h"
+#include "Engine/Font.h"
+#include "Engine/Engine.h"
+#include "Engine/StaticMesh.h"
+#include "Fonts/FontCacheFreeType.h"
+#include "GeometryBuilders/Text3DGlyphContourNode.h"
+#include "GeometryBuilders/Text3DGlyphLoader.h"
+#include "GeometryBuilders/Text3DGlyphMeshBuilder.h"
+#include "Materials/Material.h"
+#include "Misc/FileHelper.h"
+#include "UObject/ConstructorHelpers.h"
+
+FCachedFontMeshes::FCachedFontMeshes()
+{
+	CacheCounter = MakeShared<int32>(0);
+}
+
+int32 FCachedFontMeshes::GetCacheCount() const
+{
+	return CacheCounter.GetSharedReferenceCount();
+}
+
+TSharedPtr<int32> FCachedFontMeshes::GetCacheCounter()
+{
+	return CacheCounter;
+}
+
+UText3DEngineSubsystem* UText3DEngineSubsystem::Get()
+{
+	if (GEngine)
+	{
+		return GEngine->GetEngineSubsystem<UText3DEngineSubsystem>();
+	}
+
+	return nullptr;
+}
+
+UText3DEngineSubsystem::UText3DEngineSubsystem()
+{
+	if (!IsRunningDedicatedServer())
+	{
+		struct FConstructorStatics
+		{
+			ConstructorHelpers::FObjectFinder<UMaterial> Material;
+			FConstructorStatics() :
+				Material(TEXT("/Engine/BasicShapes/BasicShapeMaterial"))
+			{
+			}
+		};
+		static FConstructorStatics ConstructorStatics;
+
+		DefaultMaterial = ConstructorStatics.Material.Object;
+	}
+}
+
+void UText3DEngineSubsystem::Initialize(FSubsystemCollectionBase& Collection)
+{
+	Super::Initialize(Collection);
+	CleanupTickerHandle = FTSTicker::GetCoreTicker().AddTicker(FTickerDelegate::CreateUObject(this, &UText3DEngineSubsystem::CleanupTimerCallback), 600.0f);
+}
+
+void UText3DEngineSubsystem::Deinitialize()
+{
+	if (CleanupTickerHandle.IsValid())
+	{
+		FTSTicker::GetCoreTicker().RemoveTicker(CleanupTickerHandle);
+		CleanupTickerHandle.Reset();
+	}
+
+	Super::Deinitialize();
+}
+
+bool UText3DEngineSubsystem::Exec_Dev(UWorld* InWorld, const TCHAR* InCmd, FOutputDevice& InAr)
+{
+	if (HasAllFlags(RF_ClassDefaultObject))
+	{
+		return false;
+	}
+
+	if (FParse::Command(&InCmd, TEXT("text3d")))
+	{
+		if (FParse::Command(&InCmd, TEXT("cache")))
+		{
+			if (FParse::Command(&InCmd, TEXT("show")))
+			{
+				PrintCache();
+				return true;
+			}
+			else if (FParse::Command(&InCmd, TEXT("clear")))
+			{
+				Cleanup();
+				PrintCache();
+
+				return true;
+			}
+		}
+	}
+
+	return false;
+}
+
+void UText3DEngineSubsystem::Reset()
+{
+	CachedFonts.Reset();
+}
+
+bool UText3DEngineSubsystem::CleanupTimerCallback(float DeltaTime)
+{
+	Cleanup();
+	return true;
+}
+
+void UText3DEngineSubsystem::PrintCache() const
+{
+	UE_LOG(LogText3D, Log, TEXT("Text3D Engine Subsystem is currently caching the following fonts:"));
+	for (const TPair<uint32, FCachedFontData>& CachedFont : CachedFonts)
+	{
+		CachedFont.Value.PrintCache();
+	}
+}
+
+void UText3DEngineSubsystem::Cleanup()
+{
+	QUICK_SCOPE_CYCLE_COUNTER(STAT_UText3DEngineSubsystem_Cleanup);
+
+	UE_LOG(LogText3D, Log, TEXT("Text3D Engine Subsystem is clearing unreferenced cached fonts and typefaces."));
+	for (auto It = CachedFonts.CreateIterator(); It; ++It)
+	{
+		if (It.Value().Cleanup())
+		{
+			It.RemoveCurrent();
+		}
+	}
+}
+
+FCachedFontData& UText3DEngineSubsystem::GetCachedFontData(UFont* InFont, int32 InTypefaceEntryIndex)
+{
+	const uint32 FontHash = HashCombine(0, GetTypeHash(InFont));
+
+	if (CachedFonts.Contains(FontHash))
+	{
+		uint32 TypefaceFontDataHash = 0;
+
+		// First we check if the font itself has changed since we last cached it
+		const FCompositeFont* CompositeFont = InFont->GetCompositeFont();
+		if (CompositeFont && !CompositeFont->DefaultTypeface.Fonts.IsEmpty())
+		{
+			FTypefaceEntry Typeface;
+			if (CompositeFont->DefaultTypeface.Fonts.IsValidIndex(InTypefaceEntryIndex))
+			{
+				Typeface = CompositeFont->DefaultTypeface.Fonts[InTypefaceEntryIndex];
+			}
+			else
+			{
+				Typeface = CompositeFont->DefaultTypeface.Fonts[0];
+			}
+
+			TypefaceFontDataHash = HashCombine(0, GetTypeHash(Typeface.Font));
+		}
+
+		const uint32 CachedTypefaceFontDataHash = CachedFonts[FontHash].GetTypefaceFontDataHash(InTypefaceEntryIndex);
+		if (CachedTypefaceFontDataHash != TypefaceFontDataHash)
+		{
+			// FontFace has changed, remove the cached data
+			CachedFonts[FontHash].CleanupTypeface(InTypefaceEntryIndex);
+		}
+	}
+
+	if (!CachedFonts.Contains(FontHash))
+	{
+		FCachedFontData& NewCachedFontData = CachedFonts.Add(FontHash);
+		NewCachedFontData.SetFont(InFont);
+	}
+
+	FCachedFontData& CachedFontData = CachedFonts[FontHash];
+	CachedFontData.LoadFreeTypeFace(InTypefaceEntryIndex);
+
+	return CachedFonts[FontHash];
+}
+
+FTypefaceFontData::FTypefaceFontData()
+	: FreeTypeFace(nullptr)
+	, TypefaceFontDataHash(0)
+{
+	CacheCounter = MakeShared<int32>(0);
+}
+
+FCachedFontData::FCachedFontData()
+{
+	Font = nullptr;
+}
+
+FCachedFontData::~FCachedFontData()
+{
+	ClearFreeTypeFace();
+}
+
+void FCachedFontData::ClearFreeTypeFace()
+{
+	TypefaceFontDataMap.Reset();
+}
+
+void FCachedFontData::LoadFreeTypeFace(uint32 InTypefaceIndex)
+{
+	const bool bContainsTypeface = TypefaceFontDataMap.Contains(InTypefaceIndex);
+	if (bContainsTypeface)
+	{
+		return;
+	}
+
+	if (!Font)
+	{
+		return;
+	}
+
+	const FCompositeFont* const CompositeFont = Font->GetCompositeFont();
+	if (!CompositeFont || CompositeFont->DefaultTypeface.Fonts.IsEmpty())
+	{
+		return;
+	}
+
+	if (!CompositeFont->DefaultTypeface.Fonts.IsValidIndex(InTypefaceIndex))
+	{
+		return;
+	}
+
+	const FTypefaceEntry& Typeface = CompositeFont->DefaultTypeface.Fonts[InTypefaceIndex];
+	const FFontFaceDataConstPtr FaceData = Typeface.Font.GetFontFaceData();
+
+	if (!bContainsTypeface)
+	{
+		FTypefaceFontData& TypefaceFontData = TypefaceFontDataMap.Add(InTypefaceIndex);
+		TypefaceFontData.SetTypefaceName(Typeface.Name);
+	}
+
+	if (FaceData.IsValid() && FaceData->HasData() && FaceData->GetData().Num() > 0)
+	{
+		TypefaceFontDataMap[InTypefaceIndex].SetTypefaceData(FaceData->GetData());
+		FT_New_Memory_Face(FText3DModule::GetFreeTypeLibrary(), TypefaceFontDataMap[InTypefaceIndex].GetTypefaceData().GetData(), TypefaceFontDataMap[InTypefaceIndex].GetTypefaceData().Num(), 0, &TypefaceFontDataMap[InTypefaceIndex].GetTypeface());
+	}
+	else if (FFileHelper::LoadFileToArray(TypefaceFontDataMap[InTypefaceIndex].GetTypefaceData(), *Typeface.Font.GetFontFilename()) && TypefaceFontDataMap[InTypefaceIndex].GetTypefaceData().Num() > 0)
+	{
+		FT_New_Memory_Face(FText3DModule::GetFreeTypeLibrary(), TypefaceFontDataMap[InTypefaceIndex].GetTypefaceData().GetData(), TypefaceFontDataMap[InTypefaceIndex].GetTypefaceData().Num(), 0, &TypefaceFontDataMap[InTypefaceIndex].GetTypeface());
+	}
+
+	if (const FT_Face* FreeTypeface = &TypefaceFontDataMap[InTypefaceIndex].GetTypeface())
+	{
+		TypefaceFontDataMap[InTypefaceIndex].SetTypefaceFontDataHash(HashCombine(0, GetTypeHash(Typeface.Font)));
+		FT_Set_Char_Size(*FreeTypeface, FontSize, FontSize, 96, 96);
+		FT_Set_Pixel_Sizes(*FreeTypeface, FontSize, FontSize);
+	}
+}
+
+uint32 FCachedFontData::GetTypefaceFontDataHash(int32 InTypefaceEntryIndex) const
+{
+	if (TypefaceFontDataMap.Contains(InTypefaceEntryIndex))
+	{
+		return TypefaceFontDataMap[InTypefaceEntryIndex].GetTypefaceFontDataHash();
+	}
+
+	return 0;
+}
+
+TSharedPtr<int32> FCachedFontData::GetCacheCounter(int32 InTypefaceEntryIndex)
+{
+	if (TypefaceFontDataMap.Contains(InTypefaceEntryIndex))
+	{
+		return TypefaceFontDataMap[InTypefaceEntryIndex].GetCacheCounter();
+	}
+
+	return nullptr;
+}
+
+TSharedPtr<int32> FCachedFontData::GetMeshesCacheCounter(const FGlyphMeshParameters& InParameters)
+{
+	const uint32 HashParameters = GetTypeHash(InParameters);
+
+	const uint32 TypefaceIndex = InParameters.TypefaceIndex;
+	if (TypefaceFontDataMap.Contains(TypefaceIndex))
+	{
+		FCachedFontMeshes& CachedMeshes = TypefaceFontDataMap[TypefaceIndex].FindOrAddMeshes(HashParameters);
+		return CachedMeshes.GetCacheCounter();
+	}
+
+	return nullptr;
+}
+
+UStaticMesh* FCachedFontData::GetGlyphMesh(uint32 InGlyphIndex, const FGlyphMeshParameters& InParameters, const FFreeTypeFace* InFontFaceData)
+{
+	const uint32 TypefaceIndex = InParameters.TypefaceIndex;
+	const uint32 HashParameters = GetTypeHash(InParameters);
+
+	if (!TypefaceFontDataMap.Contains(TypefaceIndex))
+	{
+		return nullptr;
+	}
+
+	FCachedFontMeshes& CachedMeshes = TypefaceFontDataMap[TypefaceIndex].FindOrAddMeshes(HashParameters);
+
+	TObjectPtr<UStaticMesh>* CachedStaticMesh = CachedMeshes.Glyphs.Find(InGlyphIndex);
+	if (CachedStaticMesh)
+	{
+		return *CachedStaticMesh;
+	}
+
+	uint32 HashGroup = 0;
+	HashGroup = HashCombine(HashGroup, GetTypeHash(Font));
+	HashGroup = HashCombine(HashGroup, GetTypeHash(InGlyphIndex));
+	const FString StaticMeshName = FString::Printf(TEXT("Text3D_Char_%u_%u"), HashGroup, HashParameters);
+
+	const TText3DGlyphContourNodeShared Root = GetGlyphContours(InGlyphIndex, TypefaceIndex, InFontFaceData);
+	if (Root->Children.Num() == 0)
+	{
+		return nullptr;
+	}
+
+	UText3DEngineSubsystem* Subsystem = UText3DEngineSubsystem::Get();
+	const FName StrokeShapeUniqueName = MakeUniqueObjectName(Subsystem, UStaticMesh::StaticClass(), FName(*StaticMeshName));
+
+	UStaticMesh* StaticMesh = NewObject<UStaticMesh>(Subsystem, StrokeShapeUniqueName);
+	CachedMeshes.Glyphs.Add(InGlyphIndex, StaticMesh);
+
+	FText3DGlyphMeshBuilder MeshCreator;
+	MeshCreator.CreateMeshes(Root, InParameters.Extrude, InParameters.Bevel, InParameters.BevelType, InParameters.BevelSegments, InParameters.bOutline, InParameters.OutlineExpand);
+	MeshCreator.SetFrontAndBevelTextureCoordinates(InParameters.Bevel);
+	MeshCreator.MirrorGroups(InParameters.Extrude);
+
+	MeshCreator.BuildMesh(StaticMesh, Subsystem->DefaultMaterial);
+
+	return StaticMesh;
+}
+
+FT_Face FCachedFontData::GetFreeTypeFace(uint32 InTypefaceIndex)
+{
+	if (TypefaceFontDataMap.Contains(InTypefaceIndex))
+	{
+		return TypefaceFontDataMap[InTypefaceIndex].GetTypeface();
+	}
+
+	return nullptr;
+}
+
+FString FCachedFontData::GetFontName() const
+{
+	if (Font)
+	{
+		return Font->GetName();
+	}
+
+	return TEXT("");
+}
+
+TText3DGlyphContourNodeShared FCachedFontData::GetGlyphContours(uint32 InGlyphIndex, int32 InTypefaceEntryIndex, const FFreeTypeFace* InFontFaceData)
+{
+	const bool bFontFaceValid = InFontFaceData && InFontFaceData->IsFaceValid();
+	if (!bFontFaceValid && !GetFreeTypeFace(InTypefaceEntryIndex))
+	{
+		return nullptr;
+	}
+
+	const FT_Face& GlyphFreeTypeFace = bFontFaceValid
+		? InFontFaceData->GetFace()
+		: GetFreeTypeFace(InTypefaceEntryIndex);
+	check(GlyphFreeTypeFace);
+
+	if (FT_Load_Glyph(GlyphFreeTypeFace, InGlyphIndex, FT_LOAD_NO_HINTING | FT_LOAD_NO_BITMAP))
+	{
+		return nullptr;
+	}
+
+#if UE_BUILD_DEBUG
+	// If true, the symbol itself was probably bitmap-only, and this will have placeholder contours
+	bool bIsLastResort = InFontFaceData && InFontFaceData->GetFace()->family_name
+		? FString(InFontFaceData->GetFace()->family_name) == TEXT("LastResort")
+		: false;
+	if (bIsLastResort)
+	{
+		UE_LOG(LogText3D, Warning, TEXT("The following glyph index is bitmap-only: %i"), InGlyphIndex);
+	}
+#endif
+
+	FText3DGlyphLoader GlyphLoader(GlyphFreeTypeFace->glyph);
+	TText3DGlyphContourNodeShared Root = GlyphLoader.GetContourList();
+
+	return Root;
+}
+
+void FCachedFontData::PrintCache() const
+{
+	UE_LOG(LogText3D, Log, TEXT("\n== Cached typefaces for font: %s =="), *GetFontName());
+
+	for (const TPair<uint32, FTypefaceFontData>& TypefaceData : TypefaceFontDataMap)
+	{
+		UE_LOG(LogText3D, Log, TEXT("\t- %s"), *TypefaceData.Value.GetTypefaceName().ToString());
+	}
+}
+
+bool FCachedFontData::Cleanup()
+{
+	TArray<uint32> TypefaceIndexes;
+	TypefaceFontDataMap.GetKeys(TypefaceIndexes);
+
+	for (const uint32 TypefaceIndex : TypefaceIndexes)
+	{
+		CleanupTypeface(TypefaceIndex);
+	}
+
+	if (TypefaceFontDataMap.IsEmpty())
+	{
+		return true;
+	}
+
+	return false;
+}
+
+bool FCachedFontData::CleanupTypeface(uint32 InTypefaceIndex)
+{
+	bool bRemoved = false;
+	if (TypefaceFontDataMap.Contains(InTypefaceIndex))
+	{
+		FTypefaceFontData& TypefaceFontData = TypefaceFontDataMap[InTypefaceIndex];
+
+		for (auto MeshesIterator = TypefaceFontData.GetMeshes().CreateIterator(); MeshesIterator; ++MeshesIterator)
+		{
+			if (MeshesIterator.Value().GetCacheCount() <= 1)
+			{
+				MeshesIterator.RemoveCurrent();
+			}
+		}
+
+		if (TypefaceFontData.GetCacheCount() <= 1)
+		{
+			TypefaceFontData.Reset();
+
+			bRemoved = true;
+		}
+	}
+
+	if (bRemoved)
+	{
+		TypefaceFontDataMap.Remove(InTypefaceIndex);
+	}
+
+	return bRemoved;
+}
diff --git a/Engine/Plugins/Experimental/Text3D/Source/Text3D/Private/Subsystems/Text3DEngineSubsystem.h b/Engine/Plugins/Experimental/Text3D/Source/Text3D/Private/Subsystems/Text3DEngineSubsystem.h
new file mode 100644
--- /dev/null
+++ b/Engine/Plugins/Experimental/Text3D/Source/Text3D/Private/Subsystems/Text3DEngineSubsystem.h
@@ -0,0 +1,218 @@
+// Copyright Epic Games, Inc. All Rights Reserved.
+
+#pragma once
+
+#include "Containers/Ticker.h"
+#include "GeometryBuilders/Text3DGlyphContourNode.h"
+#include "Subsystems/EngineSubsystem.h"
+#include "Text3DModule.h"
+#include "Text3DTypes.h"
+
+#include "Text3DEngineSubsystem.generated.h"
+
+class FFreeTypeFace;
+class UFont;
+class UMaterial;
+class UStaticMesh;
+
+USTRUCT()
+struct FGlyphMeshParameters
+{
+	GENERATED_BODY()
+
+public:
+	UPROPERTY()
+	float Extrude = 5.0f;
+
+	UPROPERTY()
+	float Bevel = 0.0f;
+
+	UPROPERTY()
+	EText3DBevelType BevelType = EText3DBevelType::Convex;
+
+	UPROPERTY()
+	int32 BevelSegments = 8;
+
+	UPROPERTY()
+	bool bOutline = false;
+
+	UPROPERTY()
+	float OutlineExpand = 0.5f;
+
+	UPROPERTY()
+	uint32 TypefaceIndex = 0;
+};
+
+inline uint32 GetTypeHash(const FGlyphMeshParameters A)
+{
+	uint32 HashParameters = 0;
+	HashParameters = HashCombine(HashParameters, GetTypeHash(A.Extrude));
+	HashParameters = HashCombine(HashParameters, GetTypeHash(A.Bevel));
+    HashParameters = HashCombine(HashParameters, GetTypeHash(A.BevelType));
+	HashParameters = HashCombine(HashParameters, GetTypeHash(A.BevelSegments));
+	HashParameters = HashCombine(HashParameters, GetTypeHash(A.bOutline));
+	HashParameters = HashCombine(HashParameters, GetTypeHash(A.OutlineExpand));
+	HashParameters = HashCombine(HashParameters, GetTypeHash(A.TypefaceIndex));
+	return HashParameters;
+}
+
+USTRUCT()
+struct FCachedFontMeshes
+{
+	GENERATED_BODY()
+
+public:
+	FCachedFontMeshes();
+
+	int32 GetCacheCount() const;
+	TSharedPtr<int32> GetCacheCounter();
+
+	UPROPERTY()
+	TMap<uint32, TObjectPtr<UStaticMesh>> Glyphs;
+
+private:
+	TSharedPtr<int32> CacheCounter;
+};
+
+USTRUCT()
+struct FTypefaceFontData
+{
+	GENERATED_BODY()
+
+public:
+	FTypefaceFontData();
+
+	void Reset()
+	{
+		FT_Done_Face(FreeTypeFace);
+		FreeTypeFace = nullptr;
+		TypefaceData.Reset();
+	}
+
+	int32 GetCacheCount() const
+	{
+		const int32 Count = CacheCounter.GetSharedReferenceCount();
+		return Count;
+	}
+
+	FT_Face& GetTypeface() { return FreeTypeFace; }
+
+	TArray<uint8>& GetTypefaceData() { return TypefaceData; }
+
+	void SetTypefaceData(const TArray<uint8>& InTypefaceData)
+	{
+		TypefaceData = InTypefaceData;
+	}
+
+	FName GetTypefaceName() const { return TypefaceName; }
+	void SetTypefaceName(FName InTypefaceName)
+	{
+		TypefaceName = InTypefaceName;
+	}
+
+	uint32 GetTypefaceFontDataHash() const { return TypefaceFontDataHash; }
+
+	void SetTypefaceFontDataHash(uint32 InDataHash)
+	{
+		TypefaceFontDataHash = InDataHash;
+	}
+
+	TMap<uint32, FCachedFontMeshes>& GetMeshes() { return Meshes; }
+
+	TSharedPtr<int32> GetCacheCounter() { return CacheCounter; }
+
+	FCachedFontMeshes& FindOrAddMeshes(uint32 InHashParameters)
+	{
+		return Meshes.FindOrAdd(InHashParameters);
+	}
+
+private:
+	UPROPERTY()
+	TMap<uint32, FCachedFontMeshes> Meshes;
+
+	FName TypefaceName;
+	FT_Face FreeTypeFace;
+	TArray<uint8> TypefaceData;
+	TSharedPtr<int32> CacheCounter;
+	uint32 TypefaceFontDataHash;
+};
+
+USTRUCT()
+struct FCachedFontData
+{
+	GENERATED_BODY()
+
+public:
+	FCachedFontData();
+	~FCachedFontData();
+
+	FT_Face GetFreeTypeFace(uint32 InTypefaceIndex);
+
+	FString GetFontName() const;
+
+	UFont* GetFont() { return Font; }
+
+	void SetFont(UFont* InFont)
+	{
+		Font = InFont;
+	}
+
+	void LoadFreeTypeFace(uint32 InTypefaceIndex);
+	void ClearFreeTypeFace();
+
+	bool Cleanup();
+	bool CleanupTypeface(uint32 InTypefaceIndex);
+
+	uint32 GetTypefaceFontDataHash(int32 InTypefaceEntryIndex) const;
+	TSharedPtr<int32> GetCacheCounter(int32 InTypefaceEntryIndex);
+	TSharedPtr<int32> GetMeshesCacheCounter(const FGlyphMeshParameters& InParameters);
+
+	UStaticMesh* GetGlyphMesh(uint32 InGlyphIndex, const FGlyphMeshParameters& InParameters, const FFreeTypeFace* InFontFaceData = nullptr);
+	TText3DGlyphContourNodeShared GetGlyphContours(uint32 InGlyphIndex, int32 InTypefaceEntryIndex, const FFreeTypeFace* InFontFaceData = nullptr);
+
+	void PrintCache() const;
+
+private:
+	UPROPERTY()
+	TObjectPtr<UFont> Font;
+
+	UPROPERTY()
+	TMap<uint32, FTypefaceFontData> TypefaceFontDataMap;
+};
+
+UCLASS()
+class UText3DEngineSubsystem : public UEngineSubsystem, public FSelfRegisteringExec
+{
+	GENERATED_BODY()
+
+public:
+	static UText3DEngineSubsystem* Get();
+
+	UText3DEngineSubsystem();
+
+	// ~Begin UEngineSubsystem
+	virtual void Initialize(FSubsystemCollectionBase& Collection) override;
+	virtual void Deinitialize() override;
+	// ~End UEngineSubsystem
+
+	// ~Begin FSelfRegisteringExec
+	virtual bool Exec_Dev(class UWorld* InWorld, const TCHAR* InCmd, FOutputDevice& InAr) override;
+	// ~End FSelfRegisteringExec
+
+	void PrintCache() const;
+
+	void Reset();
+	void Cleanup();
+	FCachedFontData& GetCachedFontData(UFont* InFont, int32 InTypefaceEntryIndex = 0);
+
+	UPROPERTY()
+	TObjectPtr<UMaterial> DefaultMaterial;
+
+private:
+	bool CleanupTimerCallback(float DeltaTime);
+
+	UPROPERTY()
+	TMap<uint32, FCachedFontData> CachedFonts;
+
+	FTSTicker::FDelegateHandle CleanupTickerHandle;
+};
diff --git a/Engine/Plugins/Experimental/Text3D/Source/Text3D/Private/Text3DCharacterTransform.cpp b/Engine/Plugins/Experimental/Text3D/Source/Text3D/Private/Text3DCharacterTransform.cpp
deleted file mode 100644
diff --git a/Engine/Plugins/Experimental/Text3D/Source/Text3D/Private/Text3DComponent.cpp b/Engine/Plugins/Experimental/Text3D/Source/Text3D/Private/Text3DComponent.cpp
--- a/Engine/Plugins/Experimental/Text3D/Source/Text3D/Private/Text3DComponent.cpp
+++ b/Engine/Plugins/Experimental/Text3D/Source/Text3D/Private/Text3DComponent.cpp
@@ -2,145 +2,35 @@
 
 #include "Text3DComponent.h"
 
-#include "Algo/Count.h"
 #include "Async/Async.h"
-#include "Components/StaticMeshComponent.h"
 #include "Engine/Engine.h"
 #include "Engine/Font.h"
 #include "Engine/StaticMesh.h"
 #include "Fonts/CompositeFont.h"
 #include "Fonts/SlateTextShaper.h"
 #include "Framework/Text/PlainTextLayoutMarshaller.h"
-#include "Glyph.h"
+#include "GeometryBuilders/Text3DGlyph.h"
+#include "GeometryBuilders/Text3DGlyphMeshBuilder.h"
 #include "Internationalization/Regex.h"
+#include "LayoutBuilders/Text3DLayout.h"
+#include "LayoutBuilders/Text3DShapedGlyphText.h"
+#include "LayoutBuilders/Text3DLayoutShaper.h"
 #include "Materials/Material.h"
-#include "MeshCreator.h"
 #include "Misc/ScopeExit.h"
 #include "Misc/TransactionObjectEvent.h"
+#include "Renderers/Text3DRendererBase.h"
+#include "Renderers/StaticMeshes/Text3DStaticMeshesRenderer.h"
 #include "Styling/StyleDefaults.h"
-#include "Text3DEngineSubsystem.h"
-#include "Text3DPrivate.h"
-#include "TextShaper.h"
+#include "Subsystems/Text3DEngineSubsystem.h"
+#include "Text3DModule.h"
 #include "UObject/ConstructorHelpers.h"
 
 #define LOCTEXT_NAMESPACE "Text3D"
 
-struct FText3DShapedText
-{
-	FText3DShapedText()
-	{
-		Reset();
-	}
-
-	void Reset()
-	{
-		LineHeight = 0.0f;
-		FontAscender = 0.0f;
-		FontDescender = 0.0f;
-		Kerning = 0.0f;
-		WordSpacing = 0.0f;
-		bWrap = false;
-		Lines.Reset();
-	}
-
-	void CalculateWidth()
-	{
-		TArray<FShapedGlyphLine> NewLines;
-		NewLines.Reserve(Lines.Num());
-		
-		for (const FShapedGlyphLine& GlyphLine : Lines)
-		{
-			FShapedGlyphLine* CurrentLine = &NewLines.AddDefaulted_GetRef();
-			TArray<FShapedGlyphEntry> CurrentWord;
-
-			float LineWidth = 0.0f;
-			const int32 GlyphCount = GlyphLine.GlyphsToRender.Num();
-			float CurrentWordLength = 0.0f;
-			
-			for (int32 GlyphIdx = 0; GlyphIdx < GlyphCount; ++GlyphIdx)
-			{
-				const FShapedGlyphEntry& CurrentGlyph = GlyphLine.GlyphsToRender[GlyphIdx];
-
-				// Trim for proper positioning
-				if (!CurrentGlyph.bIsVisible && GlyphIdx == 0)
-				{
-					continue;
-				}
-
-				const bool bWordBreak = !CurrentGlyph.bIsVisible || GlyphIdx == GlyphLine.GlyphsToRender.Num() - 1;
-				const float GlyphAdv = GlyphLine.GetAdvance(GlyphIdx, Kerning, WordSpacing);
-
-				// If we're at the end the line or at whitespace
-				if (bWrap                               // when we're wrapping
-					&& bWordBreak                       // and at a word break
-					&& LineWidth > MaxWidth             // and the current line is longer than the max
-					&& CurrentWordLength != LineWidth)  // and the line is not just a single word that we can't break
-				{
-					CurrentLine->Width = LineWidth - CurrentWordLength;
-					CurrentLine = &NewLines.AddDefaulted_GetRef();
-					LineWidth = CurrentWordLength;
-				}
-
-				CurrentWord.Add(CurrentGlyph);
-				LineWidth += GlyphAdv;
-				CurrentWordLength += GlyphAdv;
-
-				if (bWordBreak)
-				{
-					CurrentLine->GlyphsToRender.Append(CurrentWord);
-					CurrentWordLength = 0.0f;
-					CurrentWord.Empty();
-				}
-			}
-
-			CurrentLine->Width = LineWidth;
-		}
-		
-		for (FShapedGlyphLine& NewLine : NewLines)
-		{
-			if (NewLine.GlyphsToRender.IsEmpty())
-			{
-				continue;
-			}
-
-			const FShapedGlyphEntry& FirstGlyph = NewLine.GlyphsToRender[0];
-			if (!FirstGlyph.bIsVisible)
-			{
-				NewLine.Width -= NewLine.GetAdvance(0, Kerning, WordSpacing);
-				NewLine.GlyphsToRender.RemoveAt(0);
-			}
-
-			const int32 LastIndex = NewLine.GlyphsToRender.Num() - 1;
-			const FShapedGlyphEntry& LastGlyph = NewLine.GlyphsToRender[LastIndex];
-			if (!LastGlyph.bIsVisible)
-			{
-				NewLine.Width -= NewLine.GetAdvance(LastIndex, Kerning, WordSpacing);
-				NewLine.GlyphsToRender.RemoveAt(LastIndex);
-			}
-		}
-		
-		Lines = NewLines;
-	}
-
-	float LineHeight;
-	float FontAscender;
-	float FontDescender;
-	float Kerning;
-	float WordSpacing;
-	float MaxWidth;
-	bool bWrap;
-	TArray<struct FShapedGlyphLine> Lines;
-};
-
-using TTextMeshDynamicData = TArray<TUniquePtr<FText3DDynamicData>, TFixedAllocator<static_cast<int32>(EText3DGroupType::TypeCount)>>;
-
 UText3DComponent::UText3DComponent()
 	: bIsBuilding(false)
-	, ShapedText(new FText3DShapedText())
+	, ShapedText(new FText3DShapedGlyphText())
 {
-	TextRoot = CreateDefaultSubobject<USceneComponent>(TEXT("TextRoot"));
-	TextRoot->SetupAttachment(this);
-
 	if (!IsRunningDedicatedServer())
 	{
 		struct FConstructorStatics
@@ -156,6 +46,8 @@ UText3DComponent::UText3DComponent()
 		static FConstructorStatics ConstructorStatics;
 
 		Font = ConstructorStatics.Font.Object;
+		RefreshTypeface();
+
 		UMaterial* DefaultMaterial = ConstructorStatics.Material.Object;
 		FrontMaterial = DefaultMaterial;
 		BevelMaterial = DefaultMaterial;
@@ -163,47 +55,36 @@ UText3DComponent::UText3DComponent()
 		BackMaterial = DefaultMaterial;
 	}
 
+	// Legacy renderer
+	TextRendererClass = UText3DStaticMeshesRenderer::StaticClass();
 	Text = LOCTEXT("DefaultText", "Text");
-	bOutline = false;
-	OutlineExpand = 0.5f;
-	Extrude = 5.0f;
-	Bevel = 0.0f;
-	BevelType = EText3DBevelType::Convex;
-	BevelSegments = 8;
-
-	HorizontalAlignment = EText3DHorizontalTextAlignment::Left;
-	VerticalAlignment = EText3DVerticalTextAlignment::FirstLine;
-	Kerning = 0.0f;
-	LineSpacing = 0.0f;
-	WordSpacing = 0.0f;
-
-	bHasMaxWidth = false;
-	MaxWidth = 500.f;
-	bHasMaxHeight = false;
-	MaxHeight = 500.0f;
-	bScaleProportionally = true;
-
-	bFreezeBuild = false;
-	ModifyFlags = EText3DModifyFlags::All;
+}
 
-	TextScale = FVector::ZeroVector;
+void UText3DComponent::OnComponentCreated()
+{
+	Super::OnComponentCreated();
 
-	RefreshTypeface();
+	OnTextRendererClassChanged();
 }
 
 void UText3DComponent::PostLoad()
 {
+	Super::PostLoad();
+
 	// Reset so it's rebuilt (needed if re-using the component!)
 	ModifyFlags = EText3DModifyFlags::All;
 
-	Super::PostLoad();
+	OnTextRendererClassChanged();
 }
 
 void UText3DComponent::BeginDestroy()
 {
-	ClearTextMesh();
-
 	Super::BeginDestroy();
+
+	if (TextRenderer)
+	{
+		TextRenderer->Destroy();
+	}
 }
 
 bool UText3DComponent::NeedsMeshRebuild() const
@@ -304,7 +185,7 @@ void UText3DComponent::UpdateStatistics()
 {
 	Statistics = FText3DStatistics();
 
-	const FString WordString = Text.ToString();
+	const FString WordString = GetFormattedText().ToString();
 
 	const FRegexPattern WordPattern(TEXT("\\S+"));
 	FRegexMatcher Matcher(WordPattern, WordString);
@@ -337,22 +218,25 @@ void UText3DComponent::OnRegister()
 {
 	Super::OnRegister();
 
-	if (!TextRoot->IsAttachedTo(this))
+	if (TextRenderer)
 	{
-		TextRoot->AttachToComponent(this, FAttachmentTransformRules::KeepRelativeTransform);
+		TextRenderer->Create();
 	}
-
+	
 	RebuildInternal(/** AutoUpdate */true, /** CleanCache */true);
 }
 
 void UText3DComponent::OnUnregister()
 {
+	Super::OnUnregister();
+
 	if (IsBeingDestroyed())
 	{
-		ClearTextMesh();
+		if (TextRenderer)
+		{
+			TextRenderer->Destroy();
+		}
 	}
-
-	Super::OnUnregister();
 }
 
 void UText3DComponent::PostEditImport()
@@ -374,87 +258,78 @@ void UText3DComponent::PostEditChangeProperty(FPropertyChangedEvent& PropertyCha
 {
 	Super::PostEditChangeProperty(PropertyChangedEvent);
 
-	const FName Name = PropertyChangedEvent.GetPropertyName();
+	const FName MemberPropertyName = PropertyChangedEvent.GetMemberPropertyName();
 
-	static FName BevelTypePropertyName = GET_MEMBER_NAME_CHECKED(UText3DComponent, BevelType);
-	static FName BevelSegmentsPropertyName = GET_MEMBER_NAME_CHECKED(UText3DComponent, BevelSegments);
-	if (Name == BevelTypePropertyName)
+	static const TSet<FName> FontPropertyNames
 	{
-		switch (BevelType)
-		{
-		case EText3DBevelType::Linear:
-		case EText3DBevelType::OneStep:
-		case EText3DBevelType::TwoSteps:
-		case EText3DBevelType::Engraved:
-		{
-			BevelSegments = 1;
-			break;
-		}
-		case EText3DBevelType::Convex:
-		case EText3DBevelType::Concave:
-		{
-			BevelSegments = 8;
-			break;
-		}
-		case EText3DBevelType::HalfCircle:
-		{
-			BevelSegments = 16;
-			break;
-		}
-		}
+		GET_MEMBER_NAME_CHECKED(UText3DComponent, Font),
+		GET_MEMBER_NAME_CHECKED(UText3DComponent, Typeface)
+	};
 
-		MarkForGeometryUpdate();
+	static const TSet<FName> GeometryPropertyNames
+	{
+		GET_MEMBER_NAME_CHECKED(UText3DComponent, bOutline),
+		GET_MEMBER_NAME_CHECKED(UText3DComponent, OutlineExpand),
+		GET_MEMBER_NAME_CHECKED(UText3DComponent, Extrude),
+		GET_MEMBER_NAME_CHECKED(UText3DComponent, Bevel),
+		GET_MEMBER_NAME_CHECKED(UText3DComponent, BevelType),
+		GET_MEMBER_NAME_CHECKED(UText3DComponent, BevelSegments),
+	};
+
+	static const TSet<FName> LayoutPropertyNames
+	{
+		GET_MEMBER_NAME_CHECKED(UText3DComponent, Kerning),
+		GET_MEMBER_NAME_CHECKED(UText3DComponent, LineSpacing),
+		GET_MEMBER_NAME_CHECKED(UText3DComponent, WordSpacing),
+		GET_MEMBER_NAME_CHECKED(UText3DComponent, HorizontalAlignment),
+		GET_MEMBER_NAME_CHECKED(UText3DComponent, VerticalAlignment),
+		GET_MEMBER_NAME_CHECKED(UText3DComponent, bHasMaxWidth),
+		GET_MEMBER_NAME_CHECKED(UText3DComponent, MaxWidth),
+		GET_MEMBER_NAME_CHECKED(UText3DComponent, bHasMaxHeight),
+		GET_MEMBER_NAME_CHECKED(UText3DComponent, MaxHeight),
+		GET_MEMBER_NAME_CHECKED(UText3DComponent, bScaleProportionally),
+	};
+
+	static const TSet<FName> MaterialPropertyNames
+	{
+		GET_MEMBER_NAME_CHECKED(UText3DComponent, FrontMaterial),
+		GET_MEMBER_NAME_CHECKED(UText3DComponent, BevelMaterial), 
+		GET_MEMBER_NAME_CHECKED(UText3DComponent, ExtrudeMaterial),
+		GET_MEMBER_NAME_CHECKED(UText3DComponent, BackMaterial),  
+	};
+
+	if (MemberPropertyName == GET_MEMBER_NAME_CHECKED(UText3DComponent, TextRendererClass))
+	{
+		OnTextRendererClassChanged();
 	}
-	else if (Name == BevelSegmentsPropertyName)
+	if (MemberPropertyName == GET_MEMBER_NAME_CHECKED(UText3DComponent, Text))
 	{
-		// Force minimum bevel segments based on the BevelType
-		SetBevelSegments(BevelSegments);
+		OnTextChanged();
 	}
-	else if (Name == GET_MEMBER_NAME_CHECKED(UText3DComponent, Font))
+	else if (FontPropertyNames.Contains(MemberPropertyName))
 	{
-		MarkForGeometryUpdate();
-		RefreshTypeface();
+		OnFontPropertiesChanged();
 	}
-	else if (Name == GET_MEMBER_NAME_CHECKED(UText3DComponent, Typeface) ||
-			 Name == GET_MEMBER_NAME_CHECKED(UText3DComponent, Text) ||
-			 Name == GET_MEMBER_NAME_CHECKED(UText3DComponent, OutlineExpand) ||
-			 Name == GET_MEMBER_NAME_CHECKED(UText3DComponent, bOutline) ||
-			 Name == GET_MEMBER_NAME_CHECKED(UText3DComponent, Extrude) ||
-			 Name == GET_MEMBER_NAME_CHECKED(UText3DComponent, Bevel))
+	else if (GeometryPropertyNames.Contains(MemberPropertyName))
 	{
-		MarkForGeometryUpdate();
+		OnGeometryPropertiesChanged();
 	}
-	else if (Name == GET_MEMBER_NAME_CHECKED(UText3DComponent, HorizontalAlignment) ||
-			 Name == GET_MEMBER_NAME_CHECKED(UText3DComponent, VerticalAlignment) ||
-			 Name == GET_MEMBER_NAME_CHECKED(UText3DComponent, Kerning) ||
-			 Name == GET_MEMBER_NAME_CHECKED(UText3DComponent, LineSpacing) ||
-			 Name == GET_MEMBER_NAME_CHECKED(UText3DComponent, WordSpacing) ||
-			 Name == GET_MEMBER_NAME_CHECKED(UText3DComponent, bHasMaxWidth) ||
-			 Name == GET_MEMBER_NAME_CHECKED(UText3DComponent, MaxWidth) ||
-			 Name == GET_MEMBER_NAME_CHECKED(UText3DComponent, MaxWidthHandling) ||
-			 Name == GET_MEMBER_NAME_CHECKED(UText3DComponent, bHasMaxHeight) ||
-			 Name == GET_MEMBER_NAME_CHECKED(UText3DComponent, MaxHeight) ||
-			 Name == GET_MEMBER_NAME_CHECKED(UText3DComponent, bScaleProportionally))
-	{
-		if (MaxWidthHandling == EText3DMaxWidthHandling::WrapAndScale ||
- 		   Name == GET_MEMBER_NAME_CHECKED(UText3DComponent, MaxWidthHandling))
-		{
-			MarkForGeometryUpdate();
-		}
-		else
-		{
-			MarkForLayoutUpdate();
-		}
+	else if (LayoutPropertyNames.Contains(MemberPropertyName))
+	{
+		OnLayoutPropertiesChanged();
+	}
+	else if (MemberPropertyName == GET_MEMBER_NAME_CHECKED(UText3DComponent, MaxWidthHandling))
+	{
+		OnLayoutHandlingPropertiesChanged();
 	}
-	else if (Name == GET_MEMBER_NAME_CHECKED(UText3DComponent, FrontMaterial)
-		|| Name == GET_MEMBER_NAME_CHECKED(UText3DComponent, BevelMaterial)
-		|| Name == GET_MEMBER_NAME_CHECKED(UText3DComponent, ExtrudeMaterial)
-		|| Name == GET_MEMBER_NAME_CHECKED(UText3DComponent, BackMaterial))
+	else if (MaterialPropertyNames.Contains(MemberPropertyName))
 	{
 		OnMaterialChanged();
 	}
-
-	RebuildInternal();
+	else if (MemberPropertyName == GET_MEMBER_NAME_CHECKED(UText3DComponent, bCastShadow))
+	{
+		OnVisibilityPropertiesChanged();
+	}
 }
 
 void UText3DComponent::PostTransacted(const FTransactionObjectEvent& TransactionEvent)
@@ -469,11 +344,6 @@ void UText3DComponent::PostTransacted(const FTransactionObjectEvent& Transaction
 }
 #endif
 
-bool UText3DComponent::RefreshesOnChange() const
-{
-	return bRefreshOnChange;
-}
-
 void UText3DComponent::SetRefreshOnChange(const bool Value)
 {
 	if (bRefreshOnChange != Value)
@@ -482,119 +352,94 @@ void UText3DComponent::SetRefreshOnChange(const bool Value)
 	}
 }
 
-const FText& UText3DComponent::GetText() const
-{
-	return Text;
-}
-
 void UText3DComponent::SetText(const FText& Value)
 {
-	if (!Text.EqualTo(Value))
+	if (Text.EqualTo(Value))
 	{
-		Text = Value;
-		MarkForGeometryUpdate();
-		RebuildInternal();
+		return;
 	}
-}
 
-const UFont* UText3DComponent::GetFont() const
-{
-	return Font;
+	Text = Value;
+	OnTextChanged();
 }
 
-void UText3DComponent::SetFont(UFont* InFont)
+FText UText3DComponent::GetFormattedText()
 {
-	if (Font != InFont)
+	if (!FormattedText.IsSet())
 	{
-		Font = InFont;
-		RefreshTypeface();
-		MarkForGeometryUpdate();
-		RebuildInternal();
+		FormattedText = Text;
+		FormatText(FormattedText.GetValue());
 	}
+
+	return FormattedText.GetValue();
 }
 
-bool UText3DComponent::HasOutline() const
+void UText3DComponent::SetFont(UFont* InFont)
 {
-	return bOutline;
+	if (!InFont || Font == InFont)
+	{
+		return;
+	}
+	
+	Font = InFont;
+	OnFontPropertiesChanged();
 }
 
 void UText3DComponent::SetHasOutline(const bool bValue)
 {
-	if (bOutline != bValue)
+	if (bOutline == bValue)
 	{
-		bOutline = bValue;
-		MarkForGeometryUpdate();
-		RebuildInternal();
+		return;
 	}
-}
 
-float UText3DComponent::GetOutlineExpand() const
-{
-	return OutlineExpand;
+	bOutline = bValue;
+	OnGeometryPropertiesChanged();
 }
 
 void UText3DComponent::SetOutlineExpand(const float Value)
 {
-	const float NewValue = Value;
-	if (!FMath::IsNearlyEqual(OutlineExpand, NewValue))
+	if (FMath::IsNearlyEqual(OutlineExpand, Value))
 	{
-		OutlineExpand = NewValue;
-		MarkForGeometryUpdate();
-		RebuildInternal();
+		return;
 	}
-}
 
-float UText3DComponent::GetExtrude() const
-{
-	return Extrude;
+	OutlineExpand = Value;
+	OnGeometryPropertiesChanged();
 }
 
 void UText3DComponent::SetExtrude(const float Value)
 {
 	const float NewValue = FMath::Max(0.0f, Value);
-	if (!FMath::IsNearlyEqual(Extrude, NewValue))
+	if (FMath::IsNearlyEqual(Extrude, NewValue))
 	{
-		Extrude = NewValue;
-		MarkForGeometryUpdate();
-		CheckBevel();
-		RebuildInternal();
+		return;
 	}
-}
 
-float UText3DComponent::GetBevel() const
-{
-	return Bevel;
+	Extrude = NewValue;
+	OnGeometryPropertiesChanged();
 }
 
 void UText3DComponent::SetBevel(const float Value)
 {
-	const float NewValue = FMath::Clamp(Value, 0.f, MaxBevel());
-	if (!FMath::IsNearlyEqual(Bevel, NewValue))
+	const float NewValue = FMath::Clamp(Value, 0.f, GetMaxBevel());
+	if (FMath::IsNearlyEqual(Bevel, NewValue))
 	{
-		Bevel = NewValue;
-		MarkForGeometryUpdate();
-		RebuildInternal();
+		return;
 	}
-}
-
-EText3DBevelType UText3DComponent::GetBevelType() const
-{
-	return BevelType;
+	
+	Bevel = NewValue;
+	OnGeometryPropertiesChanged();
 }
 
 void UText3DComponent::SetBevelType(const EText3DBevelType Value)
 {
-	if (BevelType != Value)
+	if (BevelType == Value)
 	{
-		BevelType = Value;
-		MarkForGeometryUpdate();
-		RebuildInternal();
+		return;
 	}
-}
 
-int32 UText3DComponent::GetBevelSegments() const
-{
-	return BevelSegments;
+	BevelType = Value;
+	OnGeometryPropertiesChanged();
 }
 
 void UText3DComponent::SetBevelSegments(const int32 Value)
@@ -606,17 +451,13 @@ void UText3DComponent::SetBevelSegments(const int32 Value)
 	}
 
 	const int32 NewValue = FMath::Clamp(Value, MinBevelSegments, 15);
-	if (BevelSegments != NewValue)
+	if (BevelSegments == NewValue)
 	{
-		BevelSegments = NewValue;
-		MarkForGeometryUpdate();
-		RebuildInternal();
+		return;
 	}
-}
-
-UMaterialInterface* UText3DComponent::GetFrontMaterial() const
-{
-	return FrontMaterial;
+	
+	BevelSegments = NewValue;
+	OnGeometryPropertiesChanged();
 }
 
 void UText3DComponent::SetFrontMaterial(UMaterialInterface* Value)
@@ -624,101 +465,21 @@ void UText3DComponent::SetFrontMaterial(UMaterialInterface* Value)
 	SetMaterial(EText3DGroupType::Front, Value);
 }
 
-UMaterialInterface* UText3DComponent::GetBevelMaterial() const
-{
-	return BevelMaterial;
-}
-
 void UText3DComponent::SetBevelMaterial(UMaterialInterface* Value)
 {
 	SetMaterial(EText3DGroupType::Bevel, Value);
 }
 
-UMaterialInterface* UText3DComponent::GetExtrudeMaterial() const
-{
-	return ExtrudeMaterial;
-}
-
 void UText3DComponent::SetExtrudeMaterial(UMaterialInterface* Value)
 {
 	SetMaterial(EText3DGroupType::Extrude, Value);
 }
 
-UMaterialInterface* UText3DComponent::GetBackMaterial() const
-{
-	return BackMaterial;
-}
-
 void UText3DComponent::SetBackMaterial(UMaterialInterface* Value)
 {
 	SetMaterial(EText3DGroupType::Back, Value);
 }
 
-bool UText3DComponent::AllocateGlyphs(int32 Num)
-{
-	int32 DeltaNum = Num - CharacterMeshes.Num();
-	if (DeltaNum == 0)
-	{
-		return false;
-	}
-	// Add characters
-	if (FMath::Sign(DeltaNum) > 0)
-	{
-		int32 GlyphId = CharacterMeshes.Num() - 1;
-		for(int32 CharacterIndex = 0; CharacterIndex < DeltaNum; ++CharacterIndex)
-		{
-			GlyphId++;
-
-			const FName CharacterKerningComponentName = MakeUniqueObjectName(this, USceneComponent::StaticClass(), FName(*FString::Printf(TEXT("CharacterKerning%d"), GlyphId)));
-			USceneComponent* CharacterKerningComponent = NewObject<USceneComponent>(this, CharacterKerningComponentName, RF_Transient);
-
-			CharacterKerningComponent->AttachToComponent(TextRoot, FAttachmentTransformRules::KeepRelativeTransform);
-			CharacterKerningComponent->RegisterComponent();
-			CharacterKernings.Add(CharacterKerningComponent);
-
-			const FName StaticMeshComponentName = MakeUniqueObjectName(this, UStaticMeshComponent::StaticClass(), FName(*FString::Printf(TEXT("StaticMeshComponent%d"), GlyphId)));
-			UStaticMeshComponent* StaticMeshComponent = NewObject<UStaticMeshComponent>(this, StaticMeshComponentName, RF_Transient);
-			StaticMeshComponent->RegisterComponent();
-			StaticMeshComponent->SetVisibility(GetVisibleFlag());
-			StaticMeshComponent->SetHiddenInGame(bHiddenInGame);
-			StaticMeshComponent->SetCastShadow(bCastShadow);
-			CharacterMeshes.Add(StaticMeshComponent);
-
-			StaticMeshComponent->AttachToComponent(CharacterKerningComponent, FAttachmentTransformRules::KeepRelativeTransform);
-		}
-	}
-	// Remove characters
-	else
-	{
-		AActor* OwnerActor = GetOwner();
-		DeltaNum = FMath::Abs(DeltaNum);
-		for(int32 CharacterIndex = CharacterKernings.Num() - 1 - DeltaNum; CharacterIndex < CharacterKernings.Num(); ++CharacterIndex)
-		{
-			USceneComponent* CharacterKerningComponent = CharacterKernings[CharacterIndex];
-			// If called in quick succession, may already be pending destruction
-			if (IsValid(CharacterKerningComponent))
-			{
-				CharacterKerningComponent->DetachFromComponent(FDetachmentTransformRules::KeepRelativeTransform);
-				CharacterKerningComponent->UnregisterComponent();
-				CharacterKerningComponent->DestroyComponent();
-			}
-
-			UStaticMeshComponent* StaticMeshComponent = CharacterMeshes[CharacterIndex];
-			if (IsValid(StaticMeshComponent))
-			{
-				StaticMeshComponent->DetachFromComponent(FDetachmentTransformRules::KeepRelativeTransform);
-				StaticMeshComponent->UnregisterComponent();
-				StaticMeshComponent->DestroyComponent();
-			}
-		}
-
-		CharacterKernings.RemoveAt(CharacterKernings.Num() - 1 - DeltaNum, DeltaNum);
-		CharacterMeshes.RemoveAt(CharacterMeshes.Num() - 1 - DeltaNum, DeltaNum);
-	}
-
-	return true;
-}
-
 UMaterialInterface* UText3DComponent::GetMaterial(const EText3DGroupType Type) const
 {
 	switch (Type)
@@ -750,6 +511,16 @@ UMaterialInterface* UText3DComponent::GetMaterial(const EText3DGroupType Type) c
 	}
 }
 
+TSharedPtr<FText3DShapedGlyphText> UText3DComponent::GetShapedText() const
+{
+	return ShapedText;
+}
+
+UText3DRendererBase* UText3DComponent::GetTextRenderer() const
+{
+	return TextRenderer;
+}
+
 void UText3DComponent::SetMaterial(const EText3DGroupType Type, UMaterialInterface* Value)
 {
 	UMaterialInterface* OldMaterial = GetMaterial(Type);
@@ -791,108 +562,82 @@ void UText3DComponent::SetMaterial(const EText3DGroupType Type, UMaterialInterfa
 	}
 }
 
-float UText3DComponent::GetKerning() const
-{
-	return Kerning;
-}
-
 void UText3DComponent::SetKerning(const float Value)
 {
-	if (!FMath::IsNearlyEqual(Kerning, Value))
+	if (FMath::IsNearlyEqual(Kerning, Value))
 	{
-		Kerning = Value;
-		UpdateTransforms();
+		return;
 	}
-}
-
-float UText3DComponent::GetLineSpacing() const
-{
-	return LineSpacing;
+	
+	Kerning = Value;
+	OnLayoutPropertiesChanged();
 }
 
 void UText3DComponent::SetLineSpacing(const float Value)
 {
-	if (!FMath::IsNearlyEqual(LineSpacing, Value))
+	if (FMath::IsNearlyEqual(LineSpacing, Value))
 	{
-		LineSpacing = Value;
-		UpdateTransforms();
+		return;
 	}
-}
-
-float UText3DComponent::GetWordSpacing() const
-{
-	return WordSpacing;
+	
+	LineSpacing = Value;
+	OnLayoutPropertiesChanged();
 }
 
 void UText3DComponent::SetWordSpacing(const float Value)
 {
-	if (!FMath::IsNearlyEqual(WordSpacing, Value))
+	if (FMath::IsNearlyEqual(WordSpacing, Value))
 	{
-		WordSpacing = Value;
-		UpdateTransforms();
+		return;
 	}
-}
 
-EText3DHorizontalTextAlignment UText3DComponent::GetHorizontalAlignment() const
-{
-	return HorizontalAlignment;
+	WordSpacing = Value;
+	OnLayoutPropertiesChanged();
 }
 
 void UText3DComponent::SetHorizontalAlignment(const EText3DHorizontalTextAlignment Value)
 {
-	if (HorizontalAlignment != Value)
+	if (HorizontalAlignment == Value)
 	{
-		HorizontalAlignment = Value;
-		UpdateTransforms();
+		return;
 	}
-}
 
-EText3DVerticalTextAlignment UText3DComponent::GetVerticalAlignment() const
-{
-	return VerticalAlignment;
+	HorizontalAlignment = Value;
+	OnLayoutPropertiesChanged();
 }
 
 void UText3DComponent::SetVerticalAlignment(const EText3DVerticalTextAlignment Value)
 {
-	if (VerticalAlignment != Value)
+	if (VerticalAlignment == Value)
 	{
-		VerticalAlignment = Value;
-		UpdateTransforms();
+		return;
 	}
-}
-
-bool UText3DComponent::HasMaxWidth() const
-{
-	return bHasMaxWidth;
+	
+	VerticalAlignment = Value;
+	OnLayoutPropertiesChanged();
 }
 
 void UText3DComponent::SetHasMaxWidth(const bool Value)
 {
-	if (bHasMaxWidth != Value)
+	if (bHasMaxWidth == Value)
 	{
-		bHasMaxWidth = Value;
-		UpdateTransforms();
+		return;
 	}
-}
 
-float UText3DComponent::GetMaxWidth() const
-{
-	return MaxWidth;
+	bHasMaxWidth = Value;
+	OnLayoutPropertiesChanged();
 }
 
 void UText3DComponent::SetMaxWidth(const float Value)
 {
 	const float NewValue = FMath::Max(1.0f, Value);
-	if (!FMath::IsNearlyEqual(MaxWidth, NewValue))
+	if (FMath::IsNearlyEqual(MaxWidth, NewValue))
 	{
-		MaxWidth = NewValue;
-		UpdateTransforms();
+		return;
 	}
-}
-
-EText3DMaxWidthHandling UText3DComponent::GetMaxWidthHandling() const
-{
-	return MaxWidthHandling;
+	
+	MaxWidth = NewValue;
+	OnLayoutPropertiesChanged();
 }
 
 void UText3DComponent::SetMaxWidthHandling(const EText3DMaxWidthHandling Value)
@@ -903,63 +648,41 @@ void UText3DComponent::SetMaxWidthHandling(const EText3DMaxWidthHandling Value)
 	}
 
 	MaxWidthHandling = Value;
-
-	if (MaxWidthHandling == EText3DMaxWidthHandling::WrapAndScale)
-	{
-		MarkForGeometryUpdate();
-	}
-	else
-	{
-		MarkForLayoutUpdate();
-	}
-}
-
-bool UText3DComponent::HasMaxHeight() const
-{
-	return bHasMaxHeight;
+	OnLayoutHandlingPropertiesChanged();
 }
 
 void UText3DComponent::SetHasMaxHeight(const bool Value)
 {
-	if (bHasMaxHeight != Value)
+	if (bHasMaxHeight == Value)
 	{
-		bHasMaxHeight = Value;
-		UpdateTransforms();
+		return;
 	}
-}
-
-float UText3DComponent::GetMaxHeight() const
-{
-	return MaxHeight;
+	
+	bHasMaxHeight = Value;
+	OnLayoutPropertiesChanged();
 }
 
 void UText3DComponent::SetMaxHeight(const float Value)
 {
 	const float NewValue = FMath::Max(1.0f, Value);
-	if (!FMath::IsNearlyEqual(MaxHeight, NewValue))
+	if (FMath::IsNearlyEqual(MaxHeight, NewValue))
 	{
-		MaxHeight = NewValue;
-		UpdateTransforms();
+		return;
 	}
-}
 
-bool UText3DComponent::ScalesProportionally() const
-{
-	return bScaleProportionally;
+	MaxHeight = NewValue;
+	OnLayoutPropertiesChanged();
 }
 
 void UText3DComponent::SetScaleProportionally(const bool Value)
 {
-	if (bScaleProportionally != Value)
+	if (bScaleProportionally == Value)
 	{
-		bScaleProportionally = Value;
-		UpdateTransforms();
+		return;
 	}
-}
 
-bool UText3DComponent::IsFrozen() const
-{
-	return bFreezeBuild;
+	bScaleProportionally = Value;
+	OnLayoutPropertiesChanged();
 }
 
 void UText3DComponent::SetFreeze(const bool bFreeze)
@@ -975,72 +698,38 @@ void UText3DComponent::SetFreeze(const bool bFreeze)
 	}
 }
 
-bool UText3DComponent::CastsShadow() const
-{
-	return bCastShadow;
-}
-
 void UText3DComponent::SetCastShadow(bool NewCastShadow)
 {
-	if (NewCastShadow != bCastShadow)
-	{
-		bCastShadow = NewCastShadow;
-
-		for (UStaticMeshComponent* MeshComponent : CharacterMeshes)
-		{
-			MeshComponent->SetCastShadow(bCastShadow);
-		}
-
-		MarkRenderStateDirty();
-	}
-}
-
-int32 UText3DComponent::GetGlyphCount()
-{
-	return TextRoot->GetNumChildrenComponents();
-}
-
-USceneComponent* UText3DComponent::GetGlyphKerningComponent(int32 Index)
-{
-	if (!CharacterKernings.IsValidIndex(Index))
+	if (NewCastShadow == bCastShadow)
 	{
-		return nullptr;
+		return;
 	}
 
-	return CharacterKernings[Index];
-}
-
-const TArray<USceneComponent*>& UText3DComponent::GetGlyphKerningComponents()
-{
-	return CharacterKernings;
+	bCastShadow = NewCastShadow;
+	OnVisibilityPropertiesChanged();
 }
 
-UStaticMeshComponent* UText3DComponent::GetGlyphMeshComponent(int32 Index)
+void UText3DComponent::SetTypeface(const FName InTypeface)
 {
-	if (!CharacterKernings.IsValidIndex(Index))
+	if (Typeface == InTypeface || !IsTypefaceAvailable(InTypeface))
 	{
-		return nullptr;
+		return;
 	}
 
-	return CharacterMeshes[Index];
-}
-
-const TArray<UStaticMeshComponent*>& UText3DComponent::GetGlyphMeshComponents()
-{
-	return CharacterMeshes;
+	Typeface = InTypeface;
+	OnFontPropertiesChanged();
 }
 
-void UText3DComponent::SetTypeface(const FName InTypeface)
+void UText3DComponent::SetTextRendererClass(const TSubclassOf<UText3DRendererBase>& InClass)
 {
-	if (Typeface == InTypeface || !IsTypefaceAvailable(InTypeface))
+	if (!TextRendererClass.Get()
+		|| TextRendererClass == InClass)
 	{
 		return;
 	}
 
-	Typeface = InTypeface;
-
-	MarkForGeometryUpdate();
-	RebuildInternal();
+	TextRendererClass = InClass;
+	OnTextRendererClassChanged();
 }
 
 void UText3DComponent::Rebuild()
@@ -1119,7 +808,7 @@ FVector UText3DComponent::GetTextScale()
 	return TextScale;
 }
 
-FVector UText3DComponent::GetLineLocation(int32 LineIndex)
+FVector UText3DComponent::GetLineLocation(int32 LineIndex) const
 {
 	float HorizontalOffset = 0.0f, VerticalOffset = 0.0f;
 	if (LineIndex < 0 || LineIndex >= ShapedText->Lines.Num())
@@ -1168,92 +857,167 @@ void UText3DComponent::UpdateTransforms()
 
 	ShapedText->CalculateWidth();
 	CalculateTextScale();
-	const FVector Scale = GetTextScale();
-	TextRoot->SetRelativeScale3D(Scale);
 
-	int32 GlyphIndex = 0;
-	for (int32 LineIndex = 0; LineIndex < ShapedText->Lines.Num(); LineIndex++)
+	if (TextRenderer)
 	{
-		FShapedGlyphLine& Line = ShapedText->Lines[LineIndex];
-		FVector Location = GetLineLocation(LineIndex);
+		TextRenderer->Update(EText3DRendererFlags::Layout);
+	}
 
-		for (int32 LineGlyph = 0; LineGlyph < Line.GlyphsToRender.Num(); LineGlyph++)
-		{
-			const FVector CharLocation = Location;
-			Location.Y += Line.GetAdvance(LineGlyph, Kerning, WordSpacing);
-			if (!Line.GlyphsToRender[LineGlyph].bIsVisible)
-			{
-				continue;
-			}
+	ModifyFlags &= ~EText3DModifyFlags::Layout;
+}
 
-			USceneComponent* GlyphKerningComponent = GetGlyphKerningComponent(GlyphIndex);
-			if (GlyphKerningComponent)
-			{
-				GlyphKerningComponent->SetRelativeLocation(CharLocation);
-			}
+void UText3DComponent::ClearTextMesh()
+{
+	if (TextRenderer)
+	{
+		TextRenderer->Clear();
+	}
+}
 
-			GlyphIndex++;
-		}
+void UText3DComponent::TriggerInternalRebuild(const EText3DModifyFlags InModifyFlags)
+{
+	if (EnumHasAnyFlags(InModifyFlags, EText3DModifyFlags::Geometry))
+	{
+		MarkForGeometryUpdate();
 	}
 
-	ModifyFlags &= ~EText3DModifyFlags::Layout;
+	if (EnumHasAnyFlags(InModifyFlags, EText3DModifyFlags::Layout))
+	{
+		MarkForLayoutUpdate();
+	}
+
+	RebuildInternal();
 }
 
-void UText3DComponent::ClearTextMesh()
+void UText3DComponent::OnTextRendererClassChanged()
 {
-	CachedCounterReferences.Reset();
+	UText3DRendererBase* RendererCDO = TextRendererClass.Get() ? TextRendererClass->GetDefaultObject<UText3DRendererBase>() : nullptr;
 
-	for (UStaticMeshComponent* MeshComponent : CharacterMeshes)
+	if (!RendererCDO)
 	{
-		if (MeshComponent)
+		if (TextRenderer)
 		{
-			MeshComponent->DetachFromComponent(FDetachmentTransformRules::KeepRelativeTransform);
-			MeshComponent->SetStaticMesh(nullptr);
-			MeshComponent->DestroyComponent();
+			TextRendererClass = TextRenderer->GetClass();
 		}
+
+		return;
 	}
-	CharacterMeshes.Reset();
 
-	for (USceneComponent* KerningComponent : CharacterKernings)
+	FName OldRendererName = NAME_None;
+
+	if (TextRenderer)
 	{
-		if (KerningComponent)
+		if (TextRendererClass == TextRenderer->GetClass())
 		{
-			KerningComponent->DetachFromComponent(FDetachmentTransformRules::KeepRelativeTransform);
-			KerningComponent->DestroyComponent();
+			return;
 		}
+
+		OldRendererName = TextRenderer->GetName();
+
+		TextRenderer->Clear();
+		TextRenderer->Destroy();
+		TextRenderer->MarkAsGarbage();
+		TextRenderer = nullptr;
 	}
-	CharacterKernings.Reset();
 
-	if (TextRoot)
+	TextRenderer = NewObject<UText3DRendererBase>(this, TextRendererClass.Get(), RendererCDO->GetName());
+
+	UE_LOG(LogText3D, Log, TEXT("%s : Text3D renderer changed : old %s - new %s"), *GetOwner()->GetActorNameOrLabel(), *OldRendererName.ToString(), *TextRenderer->GetName().ToString());
+
+	if (IsRegistered())
 	{
-		TArray<USceneComponent*> ChildComponents;
-		constexpr bool bIncludeChildDescendants = true;
-		TextRoot->GetChildrenComponents(bIncludeChildDescendants, ChildComponents);
+		TextRenderer->Create();
+		RebuildInternal(/** AutoUpdate */true, /** CleanCache */true);
+	}
+}
+
+void UText3DComponent::OnTextChanged()
+{
+	FormattedText = Text;
+	FormatText(FormattedText.GetValue());
+	
+	MarkForGeometryUpdate();
+	RebuildInternal();
+}
+
+void UText3DComponent::OnFontPropertiesChanged()
+{
+	RefreshTypeface();
 
-		for (USceneComponent* ChildComponent : ChildComponents)
+	MarkForGeometryUpdate();
+	RebuildInternal();
+}
+
+void UText3DComponent::OnGeometryPropertiesChanged()
+{
+	// Extrude
+	
+	Extrude = FMath::Max(0.0f, Extrude);
+
+	// Bevels
+	
+	Bevel = FMath::Clamp(Bevel, 0.f, GetMaxBevel());
+
+	int32 MaxBevelSegments = 1;
+	int32 MinBevelSegments = 1;
+	switch (BevelType)
+	{
+	case EText3DBevelType::Linear:
+	case EText3DBevelType::OneStep:
+	case EText3DBevelType::TwoSteps:
+	case EText3DBevelType::Engraved:
 		{
-			if (IsValid(ChildComponent))
-			{
-				ChildComponent->DetachFromComponent(FDetachmentTransformRules::KeepRelativeTransform);
-				ChildComponent->DestroyComponent();
-			}
+			MaxBevelSegments = 1;
+			break;
+		}
+	case EText3DBevelType::Convex:
+	case EText3DBevelType::Concave:
+		{
+			MaxBevelSegments = 8;
+			break;
+		}
+	case EText3DBevelType::HalfCircle:
+		{
+			MaxBevelSegments = 16;
+			MinBevelSegments = 2;
+			break;
 		}
 	}
+
+	BevelSegments = FMath::Clamp(BevelSegments, MinBevelSegments, MaxBevelSegments);
+
+	MarkForGeometryUpdate();
+	RebuildInternal();
 }
 
-void UText3DComponent::TriggerInternalRebuild(const EText3DModifyFlags InModifyFlags)
+void UText3DComponent::OnLayoutPropertiesChanged()
 {
-	if (EnumHasAnyFlags(InModifyFlags, EText3DModifyFlags::Geometry))
+	MaxWidth = FMath::Max(1.0f, MaxWidth);
+	MaxHeight = FMath::Max(1.0f, MaxHeight);
+
+	UpdateTransforms();
+}
+
+void UText3DComponent::OnLayoutHandlingPropertiesChanged()
+{
+	if (MaxWidthHandling == EText3DMaxWidthHandling::WrapAndScale)
 	{
-		MarkForGeometryUpdate();
+		OnGeometryPropertiesChanged();
+	}
+	else
+	{
+		OnLayoutPropertiesChanged();
 	}
+}
 
-	if (EnumHasAnyFlags(InModifyFlags, EText3DModifyFlags::Layout))
+void UText3DComponent::OnVisibilityPropertiesChanged()
+{
+	if (TextRenderer)
 	{
-		MarkForLayoutUpdate();
+		TextRenderer->Update(EText3DRendererFlags::Visibility);
 	}
 
-	RebuildInternal();
+	MarkRenderStateDirty();
 }
 
 void UText3DComponent::BuildTextMesh(const bool& bCleanCache)
@@ -1292,15 +1056,21 @@ void UText3DComponent::BuildTextMeshInternal(const bool& bCleanCache)
 		return;
 	}
 
-	CheckBevel();
+	if (!TextRenderer)
+	{
+		UE_LOG(LogText3D, Error, TEXT("Text renderer is not valid"));
+		return;
+	}
 
 	ClearTextMesh();
+
 	if (!Font)
 	{
+		UE_LOG(LogText3D, Error, TEXT("Font object is not valid"));
 		return;
 	}
 
-	UText3DEngineSubsystem* Subsystem = GEngine->GetEngineSubsystem<UText3DEngineSubsystem>();
+	UText3DEngineSubsystem* Subsystem = UText3DEngineSubsystem::Get();
 
 	const uint32 TypefaceIndex = GetTypeFaceIndex();
 	FCachedFontData& CachedFontData = Subsystem->GetCachedFontData(Font, TypefaceIndex);
@@ -1311,10 +1081,6 @@ void UText3DComponent::BuildTextMeshInternal(const bool& bCleanCache)
 		return;
 	}
 
-	FGlyphMeshParameters GlyphMeshParameters{Extrude, Bevel, BevelType, BevelSegments, bOutline, OutlineExpand, TypefaceIndex};
-	CachedCounterReferences.Add(CachedFontData.GetCacheCounter(TypefaceIndex));
-	CachedCounterReferences.Add(CachedFontData.GetMeshesCacheCounter(GlyphMeshParameters));
-
 	ShapedText->Reset();
 	ShapedText->LineHeight = Face->size->metrics.height * FontInverseScale;
 	ShapedText->FontAscender = Face->size->metrics.ascender * FontInverseScale;
@@ -1342,9 +1108,13 @@ void UText3DComponent::BuildTextMeshInternal(const bool& bCleanCache)
 		TextLayoutMarshaller = FPlainTextLayoutMarshaller::Create();
 	}
 
-	FText FormattedText = Text;
-	FormatText(FormattedText);
-	FTextShaper::Get()->ShapeBidirectionalText(Style, FormattedText.ToString(), TextLayout, TextLayoutMarshaller, ShapedText->Lines);
+	if (!FormattedText.IsSet())
+	{
+		FormattedText = Text;
+		FormatText(FormattedText.GetValue());
+	}
+
+	FText3DLayoutShaper::Get()->ShapeBidirectionalText(Style, FormattedText.GetValue().ToString(), TextLayout, TextLayoutMarshaller, ShapedText->Lines);
 
 	const int32 ApproximateGlyphNum = Algo::TransformAccumulate(ShapedText->Lines,
 		[](const FShapedGlyphLine& InLine)
@@ -1353,8 +1123,7 @@ void UText3DComponent::BuildTextMeshInternal(const bool& bCleanCache)
 		},
 		0);
 
-	TMap<uint32, const FFreeTypeFace*> GlyphIndexToFontFace;
-	GlyphIndexToFontFace.Reserve(ApproximateGlyphNum);
+	ShapedText->GlyphIndexToFontFace.Reserve(ApproximateGlyphNum);
 
 	for (int32 LineIndex = 0; LineIndex < ShapedText->Lines.Num(); ++LineIndex)
 	{
@@ -1363,93 +1132,28 @@ void UText3DComponent::BuildTextMeshInternal(const bool& bCleanCache)
 			if (!GlyphEntry.FontFaceData.IsValid())
 			{
 				// Add as nullptr if not already in map
-				if (!GlyphIndexToFontFace.Contains(GlyphEntry.GlyphIndex))
+				if (!ShapedText->GlyphIndexToFontFace.Contains(GlyphEntry.GlyphIndex))
 				{
-					GlyphIndexToFontFace.Add(GlyphEntry.GlyphIndex, nullptr);
+					ShapedText->GlyphIndexToFontFace.Add(GlyphEntry.GlyphIndex, nullptr);
 				}
 
 				continue;
 			}
 
-			if (const TSharedPtr<FFreeTypeFace> FontFacePtr = GlyphEntry.FontFaceData->FontFace.Pin();
-				FontFacePtr.IsValid())
+			const TSharedPtr<FFreeTypeFace> FontFacePtr = GlyphEntry.FontFaceData->FontFace.Pin();
+
+			if (FontFacePtr.IsValid())
 			{
-				GlyphIndexToFontFace.FindOrAdd(GlyphEntry.GlyphIndex, FontFacePtr.Get());
+				ShapedText->GlyphIndexToFontFace.FindOrAdd(GlyphEntry.GlyphIndex, FontFacePtr.Get());
 			}
 		}
 	}
 
 	ShapedText->CalculateWidth();
 	CalculateTextScale();
-	TextRoot->SetRelativeScale3D(GetTextScale());
-
-	// Pre-allocate, avoid new'ing!
-	AllocateGlyphs(Algo::TransformAccumulate(ShapedText->Lines, [&](const FShapedGlyphLine& ShapedLine)
-	{
-		return Algo::CountIf(ShapedLine.GlyphsToRender, [&](const FShapedGlyphEntry& ShapedGlyph)
-		{
-			return ShapedGlyph.bIsVisible;
-		});
-	},
-	0));
-
-	int32 GlyphIndex = 0;
-	for (int32 LineIndex = 0; LineIndex < ShapedText->Lines.Num(); LineIndex++)
-	{
-		const FShapedGlyphLine& ShapedLine = ShapedText->Lines[LineIndex];
-		FVector LineLocation = GetLineLocation(LineIndex);
-
-		for (int32 LineGlyph = 0; LineGlyph < ShapedLine.GlyphsToRender.Num(); LineGlyph++)
-		{
-			FVector GlyphLocation = LineLocation;
-			LineLocation.Y += ShapedLine.GetAdvance(LineGlyph, Kerning, WordSpacing);
-
-			const FShapedGlyphEntry& ShapedGlyph = ShapedLine.GlyphsToRender[LineGlyph];
-			if (!ShapedGlyph.bIsVisible)
-			{
-				continue;
-			}
-
-			// Count even when mesh is nullptr (allocation still creates components to avoid mesh building in allocation step)
-			const int32 GlyphId = GlyphIndex++;
 
-			const FFreeTypeFace* FontFace = GlyphIndexToFontFace[ShapedGlyph.GlyphIndex];
-			UStaticMesh* CachedMesh = CachedFontData.GetGlyphMesh(ShapedGlyph.GlyphIndex, GlyphMeshParameters, FontFace);
-			if (!CachedMesh || FMath::IsNearlyZero(CachedMesh->GetBounds().SphereRadius))
-			{
-				continue;
-			}
-
-			if (CharacterMeshes.IsValidIndex(GlyphId))
-			{
-				UStaticMeshComponent* StaticMeshComponent = CharacterMeshes[GlyphId];
-				StaticMeshComponent->SetStaticMesh(CachedMesh);
-				StaticMeshComponent->SetVisibility(GetVisibleFlag());
-				StaticMeshComponent->SetHiddenInGame(bHiddenInGame);
-				StaticMeshComponent->SetCastShadow(bCastShadow);
-			}
-			else
-			{
-				// @note: This shouldn't occur, but it does under unknown circumstances (UE-164789) so it should be handled
-				UE_LOG(LogText3D, Error, TEXT("CharacterMesh not found at index %d"), GlyphId);
-			}
-
-			if (CharacterKernings.IsValidIndex(GlyphId))
-			{
-				FTransform Transform;
-				Transform.SetLocation(GlyphLocation);
-				USceneComponent* CharacterKerningComponent = CharacterKernings[GlyphId];
-				CharacterKerningComponent->SetRelativeTransform(Transform);
-			}
-			else
-			{
-				// @note: This shouldn't occur, but it does under unknown circumstances (UE-164789) so it should be handled
-				UE_LOG(LogText3D, Error, TEXT("CharacterKerning not found at index %d"), GlyphId);
-			}
-		}
-	}
+	TextRenderer->Update(EText3DRendererFlags::All);
 
-	OnMaterialChanged();
 	UpdateStatistics();
 
 	TextGeneratedNativeDelegate.Broadcast();
@@ -1463,52 +1167,19 @@ void UText3DComponent::BuildTextMeshInternal(const bool& bCleanCache)
 	}
 }
 
-void UText3DComponent::CheckBevel()
-{
-	if (Bevel > MaxBevel())
-	{
-		Bevel = MaxBevel();
-	}
-}
-
-float UText3DComponent::MaxBevel() const
+float UText3DComponent::GetMaxBevel() const
 {
 	return Extrude / 2.0f;
 }
 
 void UText3DComponent::OnMaterialChanged()
 {
-	using namespace UE::Text3D::Materials;
-
-	for (UStaticMeshComponent* StaticMeshComponent : CharacterMeshes)
+	if (TextRenderer)
 	{
-		for (int32 GroupIndex = 0; GroupIndex < static_cast<int32>(EText3DGroupType::TypeCount); GroupIndex++)
-		{
-			const int32 MaterialIndex = StaticMeshComponent->GetMaterialIndex(SlotNames[GroupIndex]);
-
-			if (MaterialIndex == INDEX_NONE)
-			{
-				continue;
-			}
-
-			UMaterialInterface* Material = GetMaterial(static_cast<EText3DGroupType>(GroupIndex));
-
-			if (Material != StaticMeshComponent->GetMaterial(MaterialIndex))
-			{
-				StaticMeshComponent->SetMaterial(MaterialIndex, Material);
-			}
-		}
+		TextRenderer->Update(EText3DRendererFlags::Material);
 	}
 }
 
-FText UText3DComponent::GetFormattedText() const
-{
-	FText FormattedText = Text;
-	FormatText(FormattedText);
-
-	return FormattedText;
-}
-
 TArray<FName> UText3DComponent::GetTypefaceNames() const
 {
 	TArray<FName> TypefaceNames;
@@ -1527,19 +1198,20 @@ TArray<FName> UText3DComponent::GetTypefaceNames() const
 void UText3DComponent::OnVisibilityChanged()
 {
 	Super::OnVisibilityChanged();
-	const bool Visibility = GetVisibleFlag();
-	for (UStaticMeshComponent* StaticMeshComponent : CharacterMeshes)
+
+	if (TextRenderer)
 	{
-		StaticMeshComponent->SetVisibility(Visibility);
+		TextRenderer->Update(EText3DRendererFlags::Visibility);
 	}
 }
 
 void UText3DComponent::OnHiddenInGameChanged()
 {
 	Super::OnHiddenInGameChanged();
-	for (UStaticMeshComponent* StaticMeshComponent : CharacterMeshes)
+
+	if (TextRenderer)
 	{
-		StaticMeshComponent->SetHiddenInGame(bHiddenInGame);
+		TextRenderer->Update(EText3DRendererFlags::Visibility);
 	}
 }
 
@@ -1547,9 +1219,9 @@ void UText3DComponent::GetBounds(FVector& Origin, FVector& BoxExtent)
 {
 	FBox Box(ForceInit);
 
-	for (const UStaticMeshComponent* StaticMeshComponent : CharacterMeshes)
+	if (TextRenderer)
 	{
-		Box += StaticMeshComponent->Bounds.GetBox();
+		Box = TextRenderer->GetBounds();
 	}
 
 	Box.GetCenterAndExtents(Origin, BoxExtent);
diff --git a/Engine/Plugins/Experimental/Text3D/Source/Text3D/Private/Text3DEngineSubsystem.cpp b/Engine/Plugins/Experimental/Text3D/Source/Text3D/Private/Text3DEngineSubsystem.cpp
deleted file mode 100644
diff --git a/Engine/Plugins/Experimental/Text3D/Source/Text3D/Private/Text3DEngineSubsystem.h b/Engine/Plugins/Experimental/Text3D/Source/Text3D/Private/Text3DEngineSubsystem.h
deleted file mode 100644
diff --git a/Engine/Plugins/Experimental/Text3D/Source/Text3D/Private/Text3DInternalTypes.h b/Engine/Plugins/Experimental/Text3D/Source/Text3D/Private/Text3DInternalTypes.h
new file mode 100644
--- /dev/null
+++ b/Engine/Plugins/Experimental/Text3D/Source/Text3D/Private/Text3DInternalTypes.h
@@ -0,0 +1,17 @@
+// Copyright Epic Games, Inc. All Rights Reserved.
+
+#pragma once
+
+#include "Containers/Array.h"
+#include "UObject/NameTypes.h"
+
+namespace UE::Text3D::Materials
+{
+	static const TArray<FLazyName> SlotNames =
+	{
+		TEXT("Front"),
+		TEXT("Bevel"),
+		TEXT("Extrude"),
+		TEXT("Back")
+	};
+}
\ No newline at end of file
diff --git a/Engine/Plugins/Experimental/Text3D/Source/Text3D/Private/Text3DModule.cpp b/Engine/Plugins/Experimental/Text3D/Source/Text3D/Private/Text3DModule.cpp
new file mode 100644
--- /dev/null
+++ b/Engine/Plugins/Experimental/Text3D/Source/Text3D/Private/Text3DModule.cpp
@@ -0,0 +1,30 @@
+// Copyright Epic Games, Inc. All Rights Reserved.
+
+#include "Text3DModule.h"
+#include "Modules/ModuleManager.h"
+
+DEFINE_LOG_CATEGORY(LogText3D);
+
+IMPLEMENT_MODULE(FText3DModule, Text3D);
+
+FText3DModule::FText3DModule()
+{
+	FreeTypeLib = nullptr;
+}
+
+void FText3DModule::StartupModule()
+{
+	FT_Init_FreeType(&FreeTypeLib);
+}
+
+void FText3DModule::ShutdownModule()
+{
+	FT_Done_FreeType(FreeTypeLib);
+	FreeTypeLib = nullptr;
+}
+
+FT_Library FText3DModule::GetFreeTypeLibrary()
+{
+	const FText3DModule& Instance = FModuleManager::LoadModuleChecked<FText3DModule>("Text3D");
+	return Instance.FreeTypeLib;
+}
\ No newline at end of file
diff --git a/Engine/Plugins/Experimental/Text3D/Source/Text3D/Private/Text3DModule.h b/Engine/Plugins/Experimental/Text3D/Source/Text3D/Private/Text3DModule.h
new file mode 100644
--- /dev/null
+++ b/Engine/Plugins/Experimental/Text3D/Source/Text3D/Private/Text3DModule.h
@@ -0,0 +1,41 @@
+// Copyright Epic Games, Inc. All Rights Reserved.
+
+#pragma once
+
+#include "Logging/LogMacros.h"
+#include "Modules/ModuleInterface.h"
+
+#if PLATFORM_COMPILER_HAS_GENERIC_KEYWORD
+#define generic __identifier(generic)
+#endif	//PLATFORM_COMPILER_HAS_GENERIC_KEYWORD
+
+THIRD_PARTY_INCLUDES_START
+#include "ft2build.h"
+#include FT_FREETYPE_H
+#include FT_ADVANCES_H
+THIRD_PARTY_INCLUDES_END
+
+#if PLATFORM_COMPILER_HAS_GENERIC_KEYWORD
+#undef generic
+#endif	//PLATFORM_COMPILER_HAS_GENERIC_KEYWORD
+
+const int32 FontSize = 64;
+const float FontInverseScale = 1.0f / FontSize;
+
+DECLARE_LOG_CATEGORY_EXTERN(LogText3D, Log, All);
+
+class FText3DModule final : public IModuleInterface
+{
+public:
+	FText3DModule();
+
+	//~ Begin IModuleInterface
+	virtual void StartupModule() override;
+	virtual void ShutdownModule() override;
+	//~ End IModuleInterface
+
+	static FT_Library GetFreeTypeLibrary();
+
+private:
+	FT_Library FreeTypeLib;
+};
diff --git a/Engine/Plugins/Experimental/Text3D/Source/Text3D/Private/Text3DPrivate.cpp b/Engine/Plugins/Experimental/Text3D/Source/Text3D/Private/Text3DPrivate.cpp
deleted file mode 100644
diff --git a/Engine/Plugins/Experimental/Text3D/Source/Text3D/Private/Text3DPrivate.h b/Engine/Plugins/Experimental/Text3D/Source/Text3D/Private/Text3DPrivate.h
deleted file mode 100644
diff --git a/Engine/Plugins/Experimental/Text3D/Source/Text3D/Private/TextShaper.cpp b/Engine/Plugins/Experimental/Text3D/Source/Text3D/Private/TextShaper.cpp
deleted file mode 100644
diff --git a/Engine/Plugins/Experimental/Text3D/Source/Text3D/Private/TextShaper.h b/Engine/Plugins/Experimental/Text3D/Source/Text3D/Private/TextShaper.h
deleted file mode 100644
diff --git a/Engine/Plugins/Experimental/Text3D/Source/Text3D/Private/Util.h b/Engine/Plugins/Experimental/Text3D/Source/Text3D/Private/Util.h
deleted file mode 100644
diff --git a/Engine/Plugins/Experimental/Text3D/Source/Text3D/Public/BevelType.h b/Engine/Plugins/Experimental/Text3D/Source/Text3D/Public/BevelType.h
deleted file mode 100644
diff --git a/Engine/Plugins/Experimental/Text3D/Source/Text3D/Public/Characters/Text3DCharacterTransform.h b/Engine/Plugins/Experimental/Text3D/Source/Text3D/Public/Characters/Text3DCharacterTransform.h
new file mode 100644
--- /dev/null
+++ b/Engine/Plugins/Experimental/Text3D/Source/Text3D/Public/Characters/Text3DCharacterTransform.h
@@ -0,0 +1,151 @@
+// Copyright Epic Games, Inc. All Rights Reserved.
+
+#pragma once
+
+#include "Components/ActorComponent.h"
+#include "Components/SceneComponent.h"
+#include "Text3DTypes.h"
+#include "Text3DCharacterTransform.generated.h"
+
+UCLASS(ClassGroup=(Text3D), HideCategories = (Collision, Tags, Activation, Cooking), Meta = (BlueprintSpawnableComponent))
+class TEXT3D_API UText3DCharacterTransform : public USceneComponent
+{
+	GENERATED_BODY()
+
+public:	
+	UText3DCharacterTransform();
+
+	virtual void OnRegister() override;
+	virtual void OnUnregister() override;
+
+#if WITH_EDITOR
+	virtual void PostEditChangeProperty(struct FPropertyChangedEvent& PropertyChangedEvent) override;
+#endif
+
+	// Location
+
+	UPROPERTY(EditAnywhere, BlueprintReadWrite, BlueprintSetter = SetLocationEnabled, Category = "Location")
+	bool bLocationEnabled;
+
+	UPROPERTY(EditAnywhere, BlueprintReadWrite, BlueprintSetter = SetLocationProgress, Category = "Location", meta = (EditCondition = "bLocationEnabled", ClampMin = 0, ClampMax = 100))
+	float LocationProgress;
+
+	UPROPERTY(EditAnywhere, BlueprintReadWrite, BlueprintSetter = SetLocationOrder, Category = "Location", meta = (EditCondition = "bLocationEnabled"))
+	EText3DCharacterEffectOrder LocationOrder;
+
+	UPROPERTY(EditAnywhere, BlueprintReadWrite, BlueprintSetter = SetLocationRange, Category = "Location", meta = (EditCondition = "bLocationEnabled", ClampMin = 0, ClampMax = 100))
+	float LocationRange;
+
+	UPROPERTY(EditAnywhere, BlueprintReadWrite, BlueprintSetter = SetLocationDistance, Category = "Location", meta = (EditCondition = "bLocationEnabled"))
+	FVector LocationDistance;
+
+	UFUNCTION(BlueprintCallable, Category = "Location")
+	void SetLocationEnabled(bool bEnabled);
+
+	UFUNCTION(BlueprintCallable, Category = "Location")
+	void SetLocationProgress(float Progress);
+
+	UFUNCTION(BlueprintCallable, Category = "Location")
+	void SetLocationOrder(EText3DCharacterEffectOrder Order);
+
+	UFUNCTION(BlueprintCallable, Category = "Location")
+	void SetLocationRange(float Range);
+
+	UFUNCTION(BlueprintCallable, Category = "Location")
+	void SetLocationDistance(FVector Distance);
+
+
+	// Scale
+
+	UPROPERTY(EditAnywhere, BlueprintReadWrite, BlueprintSetter = SetScaleEnabled, Category = "Scale")
+	bool bScaleEnabled;
+
+	UPROPERTY(EditAnywhere, BlueprintReadWrite, BlueprintSetter = SetScaleProgress, Category = "Scale", meta = (EditCondition = "bScaleEnabled", ClampMin = 0, ClampMax = 100))
+	float ScaleProgress;
+
+	UPROPERTY(EditAnywhere, BlueprintReadWrite, BlueprintSetter = SetScaleOrder, Category = "Scale", meta = (EditCondition = "bScaleEnabled"))
+	EText3DCharacterEffectOrder ScaleOrder;
+
+	UPROPERTY(EditAnywhere, BlueprintReadWrite, BlueprintSetter = SetScaleRange, Category = "Scale", meta = (EditCondition = "bScaleEnabled", ClampMin = 0, ClampMax = 100))
+	float ScaleRange;
+
+	UPROPERTY(EditAnywhere, BlueprintReadWrite, BlueprintSetter = SetScaleBegin, Category = "Scale", meta = (EditCondition = "bScaleEnabled", ClampMin = 0))
+	FVector ScaleBegin;
+
+	UPROPERTY(EditAnywhere, BlueprintReadWrite, BlueprintSetter = SetScaleEnd, Category = "Scale", meta = (EditCondition = "bScaleEnabled", ClampMin = 0))
+	FVector ScaleEnd;
+
+
+	UFUNCTION(BlueprintCallable, Category = "Scale")
+	void SetScaleEnabled(bool bEnabled);
+
+	UFUNCTION(BlueprintCallable, Category = "Scale")
+	void SetScaleProgress(float Progress);
+
+	UFUNCTION(BlueprintCallable, Category = "Scale")
+	void SetScaleOrder(EText3DCharacterEffectOrder Order);
+
+	UFUNCTION(BlueprintCallable, Category = "Scale")
+	void SetScaleRange(float Range);
+
+	UFUNCTION(BlueprintCallable, Category = "Scale")
+	void SetScaleBegin(FVector Value);
+
+	UFUNCTION(BlueprintCallable, Category = "Scale")
+	void SetScaleEnd(FVector Value);
+
+
+	// Rotate
+
+	UPROPERTY(EditAnywhere, BlueprintReadWrite, BlueprintSetter = SetRotateEnabled, Category = "Rotate")
+	bool bRotateEnabled;
+
+	UPROPERTY(EditAnywhere, BlueprintReadWrite, BlueprintSetter = SetRotateProgress, Category = "Rotate", meta = (EditCondition = "bRotateEnabled", ClampMin = 0, ClampMax = 100))
+	float RotateProgress;
+
+	UPROPERTY(EditAnywhere, BlueprintReadWrite, BlueprintSetter = SetRotateOrder, Category = "Rotate", meta = (EditCondition = "bRotateEnabled"))
+	EText3DCharacterEffectOrder RotateOrder;
+
+	UPROPERTY(EditAnywhere, BlueprintReadWrite, BlueprintSetter = SetRotateRange, Category = "Rotate", meta = (EditCondition = "bRotateEnabled", ClampMin = 0, ClampMax = 100))
+	float RotateRange;
+
+	UPROPERTY(EditAnywhere, BlueprintReadWrite, BlueprintSetter = SetRotateBegin, Category = "Rotate", meta = (EditCondition = "bRotateEnabled"))
+	FRotator RotateBegin;
+
+	UPROPERTY(EditAnywhere, BlueprintReadWrite, BlueprintSetter = SetRotateEnd, Category = "Rotate", meta = (EditCondition = "bRotateEnabled"))
+	FRotator RotateEnd;
+
+
+	UFUNCTION(BlueprintCallable, Category = "Rotate")
+	void SetRotateEnabled(bool bEnabled);
+
+	UFUNCTION(BlueprintCallable, Category = "Rotate")
+	void SetRotateProgress(float Progress);
+
+	UFUNCTION(BlueprintCallable, Category = "Rotate")
+	void SetRotateOrder(EText3DCharacterEffectOrder Order);
+
+	UFUNCTION(BlueprintCallable, Category = "Rotate")
+	void SetRotateRange(float Range);
+
+	UFUNCTION(BlueprintCallable, Category = "Rotate")
+	void SetRotateBegin(FRotator Value);
+
+	UFUNCTION(BlueprintCallable, Category = "Rotate")
+	void SetRotateEnd(FRotator Value);
+
+private:
+	class UText3DComponent* GetText3DComponent();
+
+	void ProcessEffect();
+
+	float GetEffectValue(int32 Index, int32 Total, EText3DCharacterEffectOrder Order, float Progress, float Range);
+	void GetLineParameters(float Range, EText3DCharacterEffectOrder Order, int32 Count, float & EffectOut, float & StripOut);
+	int32 GetEffectPosition(int32 Index, int32 Total, EText3DCharacterEffectOrder Order);
+
+	void ResetLocation();
+	void ResetRotate();
+	void ResetScale();
+
+	bool bInitialized;
+};
diff --git a/Engine/Plugins/Experimental/Text3D/Source/Text3D/Public/Glyph.h b/Engine/Plugins/Experimental/Text3D/Source/Text3D/Public/Glyph.h
deleted file mode 100644
diff --git a/Engine/Plugins/Experimental/Text3D/Source/Text3D/Public/Mesh.h b/Engine/Plugins/Experimental/Text3D/Source/Text3D/Public/Mesh.h
deleted file mode 100644
diff --git a/Engine/Plugins/Experimental/Text3D/Source/Text3D/Public/Renderers/Text3DRendererBase.h b/Engine/Plugins/Experimental/Text3D/Source/Text3D/Public/Renderers/Text3DRendererBase.h
new file mode 100644
--- /dev/null
+++ b/Engine/Plugins/Experimental/Text3D/Source/Text3D/Public/Renderers/Text3DRendererBase.h
@@ -0,0 +1,73 @@
+// Copyright Epic Games, Inc. All Rights Reserved.
+
+#pragma once
+
+#include "Math/Box.h"
+#include "Text3DRendererBase.generated.h"
+
+class UText3DComponent;
+
+enum class EText3DRendererFlags : uint8
+{
+	/** Update whole geometry for text */
+	Geometry = 1 << 0,
+	/** Update layout for characters (transform) */
+	Layout = 1 << 1,
+	/** Update visibility properties */
+	Visibility = 1 << 2,
+	/** Update materials slots */
+	Material = 1 << 3,
+	/** Update everything */
+	All = Geometry | Layout | Visibility | Material
+};
+
+/**
+ * Base class for a rendering implementation of Text3D
+ * The whole rendering logic should be encapsulated into an instance of this class
+ */
+UCLASS(MinimalAPI, Abstract)
+class UText3DRendererBase : public UObject
+{
+	GENERATED_BODY()
+
+public:
+	/** Allocate renderer resources */
+	void Create();
+
+	/** Update rendering state */
+	void Update(EText3DRendererFlags InFlags);
+
+	/** Clears the active rendering state */
+	void Clear();
+
+	/** Cleanup renderer resources */
+	void Destroy();
+
+	/** Get cached bounds from last update */
+	FBox GetBounds() const;
+
+	/** Get the implementation name for debug purposes */
+	virtual FName GetName() const PURE_VIRTUAL(UText3DRenderingImplementationBase::GetName, return NAME_None; );
+	
+protected:
+	/** Create and setup the implementation components, called on load or creation */
+	virtual void OnCreate() PURE_VIRTUAL(UText3DRenderingImplementationBase::OnCreate, );
+
+	/** Update rendering state of text characters, called when render state is outdated */
+	virtual void OnUpdate(EText3DRendererFlags InFlags) PURE_VIRTUAL(UText3DRenderingImplementationBase::OnUpdate, );
+
+	/** Clear rendering state and remove all visible characters */
+	virtual void OnClear() PURE_VIRTUAL(UText3DRenderingImplementationBase::OnClear, );
+
+	/** Destroy and clean the implementation components, called on destroy or deactivation */
+	virtual void OnDestroy() PURE_VIRTUAL(UText3DRenderingImplementationBase::OnDestroy, );
+
+	/** Calculate the bounds of the rendered text */
+	virtual FBox OnCalculateBounds() const PURE_VIRTUAL(UText3DRenderingImplementationBase::OnCalculateBounds, return FBox(ForceInitToZero); );
+
+	UText3DComponent* GetText3DComponent() const;
+
+private:
+	TOptional<FBox> CachedBounds;
+	bool bInitialized = false;
+};
\ No newline at end of file
diff --git a/Engine/Plugins/Experimental/Text3D/Source/Text3D/Public/Text3DActor.h b/Engine/Plugins/Experimental/Text3D/Source/Text3D/Public/Text3DActor.h
--- a/Engine/Plugins/Experimental/Text3D/Source/Text3D/Public/Text3DActor.h
+++ b/Engine/Plugins/Experimental/Text3D/Source/Text3D/Public/Text3DActor.h
@@ -2,9 +2,7 @@
 
 #pragma once
 
-#include "CoreMinimal.h"
 #include "GameFramework/Actor.h"
-
 #include "Text3DActor.generated.h"
 
 class UText3DComponent;
diff --git a/Engine/Plugins/Experimental/Text3D/Source/Text3D/Public/Text3DCharacterTransform.h b/Engine/Plugins/Experimental/Text3D/Source/Text3D/Public/Text3DCharacterTransform.h
deleted file mode 100644
diff --git a/Engine/Plugins/Experimental/Text3D/Source/Text3D/Public/Text3DComponent.h b/Engine/Plugins/Experimental/Text3D/Source/Text3D/Public/Text3DComponent.h
--- a/Engine/Plugins/Experimental/Text3D/Source/Text3D/Public/Text3DComponent.h
+++ b/Engine/Plugins/Experimental/Text3D/Source/Text3D/Public/Text3DComponent.h
@@ -3,71 +3,18 @@
 #pragma once
 
 #include "Components/SceneComponent.h"
-#include "BevelType.h"
-#include "Mesh.h"
+#include "Text3DTypes.h"
 #include "UObject/ObjectMacros.h"
-
 #include "Text3DComponent.generated.h"
 
 class FTextLayout;
 class ITextLayoutMarshaller;
 class UFont;
 class UMaterialInterface;
+class UText3DRendererBase;
+struct FText3DShapedGlyphText;
 struct FTypefaceEntry;
 
-UENUM()
-enum class EText3DVerticalTextAlignment : uint8
-{
-	FirstLine		UMETA(DisplayName = "First Line"),
-	Top				UMETA(DisplayName = "Top"),
-	Center			UMETA(DisplayName = "Center"),
-	Bottom			UMETA(DisplayName = "Bottom"),
-};
-
-UENUM()
-enum class EText3DHorizontalTextAlignment : uint8
-{
-	Left			UMETA(DisplayName = "Left"),
-	Center			UMETA(DisplayName = "Center"),
-	Right			UMETA(DisplayName = "Right"),
-};
-
-UENUM()
-enum class EText3DMaxWidthHandling : uint8
-{
-	/** Scales the text to meet the max width */
-	Scale			UMETA(DisplayName = "Scale"),
-	/** First wraps the text (if possible) and then scales to meet the max width */
-	WrapAndScale	UMETA(DisplayName = "Wrap and Scale"),
-};
-
-UENUM()
-enum class EText3DModifyFlags : uint8
-{
-	None = 0,
-	Layout = 1 << 0,
-	Geometry = 1 << 1,
-	Unfreeze = 1 << 2,
-
-	All = Layout | Geometry | Unfreeze
-};
-ENUM_CLASS_FLAGS(EText3DModifyFlags)
-
-struct FText3DWordStatistics
-{
-	/** Actual range taking into account whitespaces */
-	FTextRange ActualRange;
-
-	/** Render range not taking into account whitespaces */
-	FTextRange RenderRange;
-};
-
-struct FText3DStatistics
-{
-	TArray<FText3DWordStatistics> Words;
-	int32 WhiteSpaces;
-};
-
 UCLASS(ClassGroup = (Text3D), PrioritizeCategories = "Text Layout Geometry Materials", meta = (BlueprintSpawnableComponent))
 class TEXT3D_API UText3DComponent : public USceneComponent
 {
@@ -76,22 +23,6 @@ class TEXT3D_API UText3DComponent : public USceneComponent
 public:
 	UText3DComponent();
 
-	//~ Begin UActorComponent
-	virtual void OnRegister() override;
-	virtual void OnUnregister() override;
-	//~ End UActorComponent
-	
-	//~ Begin UObject
-	virtual void BeginDestroy() override;
-	virtual void PostLoad() override;
-	virtual void PostEditImport() override;
-	virtual void PostDuplicate(EDuplicateMode::Type InDuplicateMode) override;
-#if WITH_EDITOR
-	virtual void PostEditChangeProperty(FPropertyChangedEvent& PropertyChangedEvent) override;
-	virtual void PostTransacted(const FTransactionObjectEvent& TransactionEvent) override;
-#endif
-	//~ End UObject
-
 	/**
 	 * Delegate called after text is rebuilt
 	 */
@@ -107,219 +38,295 @@ public:
 	}
 
 	/** Get whether to allow automatic refresh/mesh generation */
-	bool RefreshesOnChange() const;
+	bool RefreshesOnChange() const
+	{
+		return bRefreshOnChange;
+	}
 
 	/** Set whether to allow automatic refresh/mesh generation */
 	void SetRefreshOnChange(const bool Value);
 
 	/** Get the text value and signal the primitives to be rebuilt */
-	const FText& GetText() const;
+	const FText& GetText() const
+	{
+		return Text;
+	}
 
 	/** Set the text value and signal the primitives to be rebuilt */
-	UFUNCTION(BlueprintCallable, Category = "Rendering|Components|Text3D", meta = (DeprecatedFunction, DeprecationMessage = "Set the property directly"))
+	UFUNCTION(BlueprintCallable, Category = "Rendering|Components|Text3D")
 	void SetText(const FText& Value);
 
+	/**
+	 * Returns the Text property, after being formatted by the FormatText virtual function.
+	 * If FormatText is not overriden, the return FText will be the same as the Text property.
+	 */
+	UFUNCTION(BlueprintCallable, Category = "Text3D")
+	FText GetFormattedText();
+
 	/** Get the text font and signal the primitives to be rebuilt */
-	const UFont* GetFont() const;
+	UFont* GetFont() const
+	{
+		return Font;
+	}
 
 	/** Set the text font and signal the primitives to be rebuilt */
-	UFUNCTION(BlueprintCallable, Category = "Rendering|Components|Text3D", meta = (DeprecatedFunction, DeprecationMessage = "Set the property directly"))
+	UFUNCTION(BlueprintCallable, Category = "Rendering|Components|Text3D")
 	void SetFont(UFont* InFont);
 
 	/** Get whether an outline is applied. */
-	bool HasOutline() const;
+	bool HasOutline() const
+	{
+		return bOutline;
+	}
 
 	/** Set whether an outline is applied. */
-	UFUNCTION(BlueprintCallable, Category = "Rendering|Components|Text3D", meta = (DeprecatedFunction, DeprecationMessage = "Set the property directly"))
+	UFUNCTION(BlueprintCallable, Category = "Rendering|Components|Text3D")
 	void SetHasOutline(const bool bValue);
 
 	/** Get the outline width. */
-	float GetOutlineExpand() const;
+	float GetOutlineExpand() const
+	{
+		return OutlineExpand;
+	}
 
 	/** Set the outline width. */
-	UFUNCTION(BlueprintCallable, Category = "Rendering|Components|Text3D", meta = (DeprecatedFunction, DeprecationMessage = "Set the property directly"))
+	UFUNCTION(BlueprintCallable, Category = "Rendering|Components|Text3D")
 	void SetOutlineExpand(const float Value);
 
 	/** Get the text extrusion size and signal the primitives to be rebuilt */
-	float GetExtrude() const;
+	float GetExtrude() const
+	{
+		return Extrude;
+	}
 
 	/** Set the text extrusion size and signal the primitives to be rebuilt */
-	UFUNCTION(BlueprintCallable, Category = "Rendering|Components|Text3D", meta = (DeprecatedFunction, DeprecationMessage = "Set the property directly"))
+	UFUNCTION(BlueprintCallable, Category = "Rendering|Components|Text3D")
 	void SetExtrude(const float Value);
 
 	/** Get the 3d bevel value */
-	float GetBevel() const;
+	float GetBevel() const
+	{
+		return Bevel;
+	}
 
 	/** Set the 3d bevel value */
-	UFUNCTION(BlueprintCallable, Category = "Rendering|Components|Text3D", meta = (DeprecatedFunction, DeprecationMessage = "Set the property directly"))
+	UFUNCTION(BlueprintCallable, Category = "Rendering|Components|Text3D")
 	void SetBevel(const float Value);
 
 	/** Get the 3d bevel type */
-	EText3DBevelType GetBevelType() const;
+	EText3DBevelType GetBevelType() const
+	{
+		return BevelType;
+	}
 
 	/** Set the 3d bevel type */
-	UFUNCTION(BlueprintCallable, Category = "Rendering|Components|Text3D", meta = (DeprecatedFunction, DeprecationMessage = "Set the property directly"))
+	UFUNCTION(BlueprintCallable, Category = "Rendering|Components|Text3D")
 	void SetBevelType(const EText3DBevelType Value);
 
 	/** Get the amount of segments that will be used to tessellate the Bevel */
-	int32 GetBevelSegments() const;
+	int32 GetBevelSegments() const
+	{
+		return BevelSegments;
+	}
 
 	/** Set the amount of segments that will be used to tessellate the Bevel */
-	UFUNCTION(BlueprintCallable, Category = "Rendering|Components|Text3D", meta = (DeprecatedFunction, DeprecationMessage = "Set the property directly"))
+	UFUNCTION(BlueprintCallable, Category = "Rendering|Components|Text3D")
 	void SetBevelSegments(const int32 Value);
 
 	/** Get the text front material */
-	UMaterialInterface* GetFrontMaterial() const;
+	UMaterialInterface* GetFrontMaterial() const
+	{
+		return FrontMaterial;
+	}
 
 	/** Set the text front material */
-	UFUNCTION(BlueprintCallable, Category = "Rendering|Components|Text3D", meta = (DeprecatedFunction, DeprecationMessage = "Set the property directly"))
+	UFUNCTION(BlueprintCallable, Category = "Rendering|Components|Text3D")
 	void SetFrontMaterial(UMaterialInterface* Value);
 
 	/** Get the text bevel material */
-	UMaterialInterface* GetBevelMaterial() const;
+	UMaterialInterface* GetBevelMaterial() const
+	{
+		return BevelMaterial;
+	}
 
 	/** Set the text bevel material */
-	UFUNCTION(BlueprintCallable, Category = "Rendering|Components|Text3D", meta = (DeprecatedFunction, DeprecationMessage = "Set the property directly"))
+	UFUNCTION(BlueprintCallable, Category = "Rendering|Components|Text3D")
 	void SetBevelMaterial(UMaterialInterface* Value);
 
 	/** Get the text extrude material */
-	UMaterialInterface* GetExtrudeMaterial() const;
+	UMaterialInterface* GetExtrudeMaterial() const
+	{
+		return ExtrudeMaterial;
+	}
 
 	/** Set the text extrude material */
-	UFUNCTION(BlueprintCallable, Category = "Rendering|Components|Text3D", meta = (DeprecatedFunction, DeprecationMessage = "Set the property directly"))
+	UFUNCTION(BlueprintCallable, Category = "Rendering|Components|Text3D")
 	void SetExtrudeMaterial(UMaterialInterface* Value);
 
 	/** Get the text back material */
-	UMaterialInterface* GetBackMaterial() const;
+	UMaterialInterface* GetBackMaterial() const
+	{
+		return BackMaterial;
+	}
 
 	/** Set the text back material */
-	UFUNCTION(BlueprintCallable, Category = "Rendering|Components|Text3D", meta = (DeprecatedFunction, DeprecationMessage = "Set the property directly"))
+	UFUNCTION(BlueprintCallable, Category = "Rendering|Components|Text3D")
 	void SetBackMaterial(UMaterialInterface* Value);
 
 	/** Get the kerning value and signal the primitives to be rebuilt */
-	float GetKerning() const;
+	float GetKerning() const
+	{
+		return Kerning;
+	}
 
 	/** Set the kerning value and signal the primitives to be rebuilt */
-	UFUNCTION(BlueprintCallable, Category = "Rendering|Components|Text3D", meta = (DeprecatedFunction, DeprecationMessage = "Set the property directly"))
+	UFUNCTION(BlueprintCallable, Category = "Rendering|Components|Text3D")
 	void SetKerning(const float Value);
 
 	/** Get the line spacing value and signal the primitives to be rebuilt */
-	float GetLineSpacing() const;
+	float GetLineSpacing() const
+	{
+		return LineSpacing;
+	}
 
 	/** Set the line spacing value and signal the primitives to be rebuilt */
-	UFUNCTION(BlueprintCallable, Category = "Rendering|Components|Text3D", meta = (DeprecatedFunction, DeprecationMessage = "Set the property directly"))
+	UFUNCTION(BlueprintCallable, Category = "Rendering|Components|Text3D")
 	void SetLineSpacing(const float Value);
 
 	/** Get the word spacing value and signal the primitives to be rebuilt */
-	float GetWordSpacing() const;
+	float GetWordSpacing() const
+	{
+		return WordSpacing;
+	}
 
 	/** Set the word spacing value and signal the primitives to be rebuilt */
-	UFUNCTION(BlueprintCallable, Category = "Rendering|Components|Text3D", meta = (DeprecatedFunction, DeprecationMessage = "Set the property directly"))
+	UFUNCTION(BlueprintCallable, Category = "Rendering|Components|Text3D")
 	void SetWordSpacing(const float Value);
 
 	/** Get the horizontal alignment value and signal the primitives to be rebuilt */
-	EText3DHorizontalTextAlignment GetHorizontalAlignment() const;
+	EText3DHorizontalTextAlignment GetHorizontalAlignment() const
+	{
+		return HorizontalAlignment;
+	}
 
 	/** Set the horizontal alignment value and signal the primitives to be rebuilt */
-	UFUNCTION(BlueprintCallable, Category = "Rendering|Components|Text3D", meta = (DeprecatedFunction, DeprecationMessage = "Set the property directly"))
+	UFUNCTION(BlueprintCallable, Category = "Rendering|Components|Text3D")
 	void SetHorizontalAlignment(const EText3DHorizontalTextAlignment value);
 
 	/** Get the vertical alignment and signal the primitives to be rebuilt */
-	EText3DVerticalTextAlignment GetVerticalAlignment() const;
+	EText3DVerticalTextAlignment GetVerticalAlignment() const
+	{
+		return VerticalAlignment;
+	}
 
 	/** Set the vertical alignment and signal the primitives to be rebuilt */
-	UFUNCTION(BlueprintCallable, Category = "Rendering|Components|Text3D", meta = (DeprecatedFunction, DeprecationMessage = "Set the property directly"))
+	UFUNCTION(BlueprintCallable, Category = "Rendering|Components|Text3D")
 	void SetVerticalAlignment(const EText3DVerticalTextAlignment value);
 
 	/** Whether a maximum width is specified */
-	bool HasMaxWidth() const;
+	bool HasMaxWidth() const
+	{
+		return bHasMaxWidth;
+	}
 
 	/** Enable / Disable a Maximum Width */
-	UFUNCTION(BlueprintCallable, Category = "Rendering|Components|Text3D", meta = (DeprecatedFunction, DeprecationMessage = "Set the property directly"))
+	UFUNCTION(BlueprintCallable, Category = "Rendering|Components|Text3D")
 	void SetHasMaxWidth(const bool Value);
 
 	/** Get the Maximum Width - If width is larger, mesh will scale down to fit MaxWidth value */
-	float GetMaxWidth() const;
+	float GetMaxWidth() const
+	{
+		return MaxWidth;
+	}
 
 	/** Set the Maximum Width - If width is larger, mesh will scale down to fit MaxWidth value */
-	UFUNCTION(BlueprintCallable, Category = "Rendering|Components|Text3D", meta = (DeprecatedFunction, DeprecationMessage = "Set the property directly"))
+	UFUNCTION(BlueprintCallable, Category = "Rendering|Components|Text3D")
 	void SetMaxWidth(const float Value);
 
 	/** Get the Maximum Width Handling - Whether to wrap before scaling when the text size reaches the max width */
-	EText3DMaxWidthHandling GetMaxWidthHandling() const;
+	EText3DMaxWidthHandling GetMaxWidthHandling() const
+	{
+		return MaxWidthHandling;
+	}
 
 	/** Set the Maximum Width Handling - Whether to wrap before scaling when the text size reaches the max width */
-	UFUNCTION(BlueprintCallable, Category = "Rendering|Components|Text3D", meta = (DeprecatedFunction, DeprecationMessage = "Set the property directly"))
+	UFUNCTION(BlueprintCallable, Category = "Rendering|Components|Text3D")
 	void SetMaxWidthHandling(const EText3DMaxWidthHandling Value);
 
 	/** Whether a maximum height is specified */
-	bool HasMaxHeight() const;
+	bool HasMaxHeight() const
+	{
+		return bHasMaxHeight;
+	}
 
 	/** Enable / Disable a Maximum Height */
-	UFUNCTION(BlueprintCallable, Category = "Rendering|Components|Text3D", meta = (DeprecatedFunction, DeprecationMessage = "Set the property directly"))
+	UFUNCTION(BlueprintCallable, Category = "Rendering|Components|Text3D")
 	void SetHasMaxHeight(const bool Value);
 
 	/** Get the Maximum Height - If height is larger, mesh will scale down to fit MaxHeight value */
-	float GetMaxHeight() const;
+	float GetMaxHeight() const
+	{
+		return MaxHeight;
+	}
 
 	/** Set the Maximum Height - If height is larger, mesh will scale down to fit MaxHeight value */
-	UFUNCTION(BlueprintCallable, Category = "Rendering|Components|Text3D", meta = (DeprecatedFunction, DeprecationMessage = "Set the property directly"))
+	UFUNCTION(BlueprintCallable, Category = "Rendering|Components|Text3D")
 	void SetMaxHeight(const float Value);
 
 	/** Get if the mesh should scale proportionally when Max Width/Height is set */
-	bool ScalesProportionally() const;
+	bool ScalesProportionally() const
+	{
+		return bScaleProportionally;
+	}
 
 	/** Set if the mesh should scale proportionally when Max Width/Height is set */
-	UFUNCTION(BlueprintCallable, Category = "Rendering|Components|Text3D", meta = (DeprecatedFunction, DeprecationMessage = "Set the property directly"))
+	UFUNCTION(BlueprintCallable, Category = "Rendering|Components|Text3D")
 	void SetScaleProportionally(const bool Value);
 
 	/** Freeze mesh rebuild, to avoid unnecessary mesh rebuilds when setting a few properties together */
-	bool IsFrozen() const;
+	bool IsFrozen() const
+	{
+		return bFreezeBuild;
+	}
 
 	/** Freeze mesh rebuild, to avoid unnecessary mesh rebuilds when setting a few properties together */
-	UFUNCTION(BlueprintCallable, Category = "Rendering|Components|Text3D", meta = (DeprecatedFunction, DeprecationMessage = "Set the property directly"))
+	UFUNCTION(BlueprintCallable, Category = "Rendering|Components|Text3D")
 	void SetFreeze(const bool bFreeze);
 
 	/** Get the value of CastShadow. */
-	bool CastsShadow() const;
+	bool CastsShadow() const
+	{
+		return bCastShadow;
+	}
 
 	/** Set the value of CastShadow. */
-	UFUNCTION(BlueprintCallable, Category = "Rendering|Components|Text3D", meta = (DeprecatedFunction, DeprecationMessage = "Set the property directly"))
+	UFUNCTION(BlueprintCallable, Category = "Rendering|Components|Text3D")
 	void SetCastShadow(bool NewCastShadow);
 
+	/** Get whole text rendered bounds */
 	UFUNCTION(BlueprintCallable, Category = "Rendering|Components|Text3D")
 	void GetBounds(FVector& Origin, FVector& BoxExtent);
 
-	/** Gets the number of font glyphs that are currently used */
-	UFUNCTION(BlueprintCallable, Category = "Rendering|Components|Text3D")
-	int32 GetGlyphCount();
-
-	/** Gets the USceneComponent that a glyph is attached to */
-	UFUNCTION(BlueprintCallable, Category = "Rendering|Components|Text3D")
-	USceneComponent* GetGlyphKerningComponent(int32 Index);
-
-	/** Gets all the glyph kerning components */
-	UFUNCTION(BlueprintCallable, Category = "Rendering|Components|Text3D")
-	const TArray<USceneComponent*>& GetGlyphKerningComponents();
-
-	/** Gets the StaticMeshComponent of a glyph */
-	UFUNCTION(BlueprintCallable, Category = "Rendering|Components|Text3D")
-	UStaticMeshComponent* GetGlyphMeshComponent(int32 Index);
-
-	/** Gets all the glyph meshes */
-	UFUNCTION(BlueprintCallable, Category = "Rendering|Components|Text3D")
-	const TArray<UStaticMeshComponent*>& GetGlyphMeshComponents();
-
 	/** Gets the scale of actual text geometry, taking into account MaxWidth and MaxHeight constraints. This function will NOT return the component scale*/
 	UFUNCTION(BlueprintCallable, Category = "Rendering|Components|Text3D")
 	FVector GetTextScale();
 
 	/** Get the typeface */
-	FName GetTypeface() const { return  Typeface; }
+	FName GetTypeface() const
+	{
+		return Typeface;
+	}
 
 	/** Set the typeface */
 	void SetTypeface(const FName InTypeface);
 
+	void SetTextRendererClass(const TSubclassOf<UText3DRendererBase>& InClass);
+	
+	TSubclassOf<UText3DRendererBase> GetTextRendererClass() const
+	{
+		return TextRendererClass;
+	}
+
 	/** Manually update the geometry, ignoring RefreshOnChange (but still accounting for the Freeze flag) */
 	void Rebuild();
 
@@ -328,144 +335,194 @@ public:
 		return Statistics;
 	}
 
+	void MarkForLayoutUpdate();
+
+	FVector GetLineLocation(int32 LineIndex) const;
+	UMaterialInterface* GetMaterial(const EText3DGroupType Type) const;
+	TSharedPtr<FText3DShapedGlyphText> GetShapedText() const;
+	uint32 GetTypeFaceIndex() const;
+
+	UText3DRendererBase* GetTextRenderer() const;
+
 protected:
-	/** Whether to allow automatic refresh/mesh generation */
-	UPROPERTY(EditAnywhere, BlueprintReadWrite, Getter = "RefreshesOnChange", Setter = "SetRefreshOnChange", Category = "Text3D", AdvancedDisplay, meta = (AllowPrivateAccess = "true"))
-	bool bRefreshOnChange = true;
+	//~ Begin UActorComponent
+	virtual void OnComponentCreated() override;
+	virtual void OnRegister() override;
+	virtual void OnUnregister() override;
+	//~ End UActorComponent
+	
+	//~ Begin UObject
+	virtual void BeginDestroy() override;
+	virtual void PostLoad() override;
+	virtual void PostEditImport() override;
+	virtual void PostDuplicate(EDuplicateMode::Type InDuplicateMode) override;
+#if WITH_EDITOR
+	virtual void PostEditChangeProperty(FPropertyChangedEvent& PropertyChangedEvent) override;
+	virtual void PostTransacted(const FTransactionObjectEvent& TransactionEvent) override;
+#endif
+	//~ End UObject
+
+	/** Get the type faces supported for the current font */
+	UFUNCTION()
+	TArray<FName> GetTypefaceNames() const;
+
+	/** Intercept and propagate a change on this component to all children. */
+	virtual void OnVisibilityChanged() override;
+
+	/** Intercept and propagate a change on this component to all children. */
+	virtual void OnHiddenInGameChanged() override;
+
+	/**
+	 * Will be called when text geometry is generated.
+	 * Override it to customize text formatting in the final geometry, without affecting the Text property.
+	 * Use GetFormattedText() to retrieve a FText with the result of this formatting.
+	 */
+	virtual void FormatText(FText& InOutText) const {}
+
+	/** Clears all generated components and meshes from this component. */
+	void ClearTextMesh();
+
+	/** Can be used to force an internal geometry and/or layout rebuild */
+	void TriggerInternalRebuild(const EText3DModifyFlags InModifyFlags);
+
+	/** Called when TextRendererClass is changed */
+	void OnTextRendererClassChanged();
+
+	/** Called when Text is changed */
+	void OnTextChanged();
+
+	/** Called when font options are changed */
+	void OnFontPropertiesChanged();
+
+	/** Called when geometry options are changed */
+	void OnGeometryPropertiesChanged();
+
+	/** Called when layout options are changed */
+	void OnLayoutPropertiesChanged();
+
+	/** Called when layout handling options are changed */
+	void OnLayoutHandlingPropertiesChanged();
+
+	/** Called when visibility options are changed */
+	void OnVisibilityPropertiesChanged();
+
+	/** The renderer class to use to create the text geometry */
+	UPROPERTY(EditAnywhere, Setter, Getter, Category="Text", NoClear, meta=(ShowDisplayNames, AllowPrivateAccess="true"))
+	TSubclassOf<UText3DRendererBase> TextRendererClass;
 
 	/** The text to generate a 3d mesh */
-	UPROPERTY(EditAnywhere, BlueprintReadWrite, Getter, Setter, Category = "Text", meta = (MultiLine = true, AllowPrivateAccess = "true"))
+	UPROPERTY(EditAnywhere, Getter, Setter, Category = "Text", meta = (MultiLine = true, AllowPrivateAccess = "true"))
 	FText Text;
 
+	/** Original text formatted */
+	TOptional<FText> FormattedText;
+
+	/** Text font defines the style of rendered characters */
+	UPROPERTY(EditAnywhere, Getter, Setter, Category = "Text", meta = (AllowPrivateAccess = "true"))
+	TObjectPtr<UFont> Font;
+
+	/** Text font face, subset within font like bold, italic, regular */
+	UPROPERTY(EditAnywhere, Getter, Setter, Category = "Text", meta = (AllowPrivateAccess = "true", GetOptions="GetTypefaceNames"))
+	FName Typeface;
+
+	/** Whether to allow automatic refresh/mesh generation */
+	UPROPERTY(EditAnywhere, Getter = "RefreshesOnChange", Setter = "SetRefreshOnChange", Category = "Text", AdvancedDisplay, meta = (AllowPrivateAccess = "true"))
+	bool bRefreshOnChange = true;
+
 	/** Size of the extrude */
-	UPROPERTY(EditAnywhere, BlueprintReadWrite, Getter, Setter, Category = "Geometry", meta = (ClampMin = 0, AllowPrivateAccess = "true"))
-	float Extrude;
+	UPROPERTY(EditAnywhere, Getter, Setter, Category = "Geometry", meta = (ClampMin = 0, AllowPrivateAccess = "true"))
+	float Extrude = 5.0f;
 
 	/** Size of bevel */
-	UPROPERTY(EditAnywhere, BlueprintReadWrite, Getter, Setter, Category = "Geometry", meta = (EditCondition = "!bOutline && Extrude > 0", ClampMin = 0, AllowPrivateAccess = "true"))
-	float Bevel;
+	UPROPERTY(EditAnywhere, Getter, Setter, Category = "Geometry", meta = (EditCondition = "!bOutline && Extrude > 0", ClampMin = 0, AllowPrivateAccess = "true"))
+	float Bevel = 0.0f;
 
 	/** Bevel Type */
-	UPROPERTY(EditAnywhere, BlueprintReadWrite, Getter, Setter, Category = "Geometry", meta = (EditCondition = "!bOutline && Extrude > 0", AllowPrivateAccess = "true"))
-	EText3DBevelType BevelType;
+	UPROPERTY(EditAnywhere, Getter, Setter, Category = "Geometry", meta = (EditCondition = "!bOutline && Extrude > 0", AllowPrivateAccess = "true"))
+	EText3DBevelType BevelType = EText3DBevelType::Convex;
 
 	/** Bevel Segments (Defines the amount of tesselation for the bevel part) */
-	UPROPERTY(EditAnywhere, BlueprintReadWrite, Getter, Setter, Category = "Geometry", meta = (EditCondition = "!bOutline && Extrude > 0", ClampMin = 1, ClampMax = 15, AllowPrivateAccess = "true"))
-	int32 BevelSegments;
+	UPROPERTY(EditAnywhere, Getter, Setter, Category = "Geometry", meta = (EditCondition = "!bOutline && Extrude > 0", ClampMin = 1, ClampMax = 15, AllowPrivateAccess = "true"))
+	int32 BevelSegments = 8;
 
 	/** Generate Outline */
-	UPROPERTY(EditAnywhere, BlueprintReadWrite, Getter = "HasOutline", Setter = "SetHasOutline", Category = "Geometry", meta = (AllowPrivateAccess = "true"))
-	bool bOutline;
+	UPROPERTY(EditAnywhere, Getter = "HasOutline", Setter = "SetHasOutline", Category = "Geometry", meta = (AllowPrivateAccess = "true"))
+	bool bOutline = false;
 
 	/** Outline expand/offset amount */
-	UPROPERTY(EditAnywhere, BlueprintReadWrite, Getter, Setter, Category = "Geometry", meta = (EditCondition = "bOutline", AllowPrivateAccess = "true"))
-	float OutlineExpand;
+	UPROPERTY(EditAnywhere, Getter, Setter, Category = "Geometry", meta = (EditCondition = "bOutline", AllowPrivateAccess = "true"))
+	float OutlineExpand = 0.5f;
 
 	/** Material for the front part */
-	UPROPERTY(EditAnywhere, BlueprintReadWrite, Getter, Setter, Category = "Materials", meta = (AllowPrivateAccess = "true"))
+	UPROPERTY(EditAnywhere, Getter, Setter, Category = "Materials", meta = (AllowPrivateAccess = "true"))
 	TObjectPtr<UMaterialInterface> FrontMaterial;
 
 	/** Material for the bevel part */
-	UPROPERTY(EditAnywhere, BlueprintReadWrite, Getter, Setter, Category = "Materials", meta = (EditCondition = "!bOutline && Extrude > 0 && Bevel > 0", EditConditionHides, AllowPrivateAccess = "true"))
+	UPROPERTY(EditAnywhere, Getter, Setter, Category = "Materials", meta = (EditCondition = "!bOutline && Extrude > 0 && Bevel > 0", EditConditionHides, AllowPrivateAccess = "true"))
 	TObjectPtr<UMaterialInterface> BevelMaterial;
 
 	/** Material for the extruded part */
-	UPROPERTY(EditAnywhere, BlueprintReadWrite, Getter, Setter, Category = "Materials", meta = (EditCondition = "Extrude > 0", EditConditionHides, AllowPrivateAccess = "true"))
+	UPROPERTY(EditAnywhere, Getter, Setter, Category = "Materials", meta = (EditCondition = "Extrude > 0", EditConditionHides, AllowPrivateAccess = "true"))
 	TObjectPtr<UMaterialInterface> ExtrudeMaterial;
 
 	/** Material for the back part */
-	UPROPERTY(EditAnywhere, BlueprintReadWrite, Getter, Setter, Category = "Materials", meta = (EditCondition = "Extrude > 0", EditConditionHides, AllowPrivateAccess = "true"))
+	UPROPERTY(EditAnywhere, Getter, Setter, Category = "Materials", meta = (EditCondition = "Extrude > 0", EditConditionHides, AllowPrivateAccess = "true"))
 	TObjectPtr<UMaterialInterface> BackMaterial;
 
-	/** Text font */
-	UPROPERTY(EditAnywhere, BlueprintReadWrite, Getter, Setter, Category = "Text", meta = (AllowPrivateAccess = "true"))
-	TObjectPtr<UFont> Font;
-
-	UPROPERTY(EditAnywhere, BlueprintReadWrite, Getter, Setter, Category = "Text", meta = (AllowPrivateAccess = "true", GetOptions="GetTypefaceNames"))
-	FName Typeface;
-
 	/** Horizontal text alignment */
-	UPROPERTY(EditAnywhere, BlueprintReadWrite, Getter, Setter, Category = "Layout", meta = (AllowPrivateAccess = "true"))
-	EText3DHorizontalTextAlignment HorizontalAlignment;
+	UPROPERTY(EditAnywhere, Getter, Setter, Category = "Layout", meta = (AllowPrivateAccess = "true"))
+	EText3DHorizontalTextAlignment HorizontalAlignment = EText3DHorizontalTextAlignment::Left;
 
 	/** Vertical text alignment */
-	UPROPERTY(EditAnywhere, BlueprintReadWrite, Getter, Setter, Category = "Layout", meta = (AllowPrivateAccess = "true"))
-	EText3DVerticalTextAlignment VerticalAlignment;
+	UPROPERTY(EditAnywhere, Getter, Setter, Category = "Layout", meta = (AllowPrivateAccess = "true"))
+	EText3DVerticalTextAlignment VerticalAlignment = EText3DVerticalTextAlignment::FirstLine;
 
 	/** Text kerning */
-	UPROPERTY(EditAnywhere, BlueprintReadWrite, Getter, Setter, Category = "Layout", meta = (AllowPrivateAccess = "true"))
-	float Kerning;
+	UPROPERTY(EditAnywhere, Getter, Setter, Category = "Layout", meta = (AllowPrivateAccess = "true"))
+	float Kerning = 0.0f;
 
 	/** Extra line spacing */
-	UPROPERTY(EditAnywhere, BlueprintReadWrite, Getter, Setter, Category = "Layout", meta = (AllowPrivateAccess = "true"))
-	float LineSpacing;
+	UPROPERTY(EditAnywhere, Getter, Setter, Category = "Layout", meta = (AllowPrivateAccess = "true"))
+	float LineSpacing = 0.0f;
 
 	/** Extra word spacing */
-	UPROPERTY(EditAnywhere, BlueprintReadWrite, Setter, Category = "Layout", meta = (AllowPrivateAccess = "true"))
-	float WordSpacing;
+	UPROPERTY(EditAnywhere, Setter, Category = "Layout", meta = (AllowPrivateAccess = "true"))
+	float WordSpacing = 0.0f;
 
 	/** Enables a maximum width to the 3D Text */
-	UPROPERTY(EditAnywhere, BlueprintReadWrite, Getter = "HasMaxWidth", Setter = "SetHasMaxWidth", Category = "Layout", meta = (InlineEditConditionToggle, AllowPrivateAccess = "true"))
-	bool bHasMaxWidth;
+	UPROPERTY(EditAnywhere, Getter = "HasMaxWidth", Setter = "SetHasMaxWidth", Category = "Layout", meta = (InlineEditConditionToggle, AllowPrivateAccess = "true"))
+	bool bHasMaxWidth = false;
 
 	/** Sets a maximum width to the 3D Text */
-	UPROPERTY(EditAnywhere, BlueprintReadWrite, Getter, Setter, Category = "Layout", meta = (EditCondition = "bHasMaxWidth", ClampMin = 1, AllowPrivateAccess = "true"))
-	float MaxWidth;
+	UPROPERTY(EditAnywhere, Getter, Setter, Category = "Layout", meta = (EditCondition = "bHasMaxWidth", ClampMin = 1, AllowPrivateAccess = "true"))
+	float MaxWidth = 500.f;
 
 	/** Enables a maximum height to the 3D Text */
-	UPROPERTY(EditAnywhere, BlueprintReadWrite, Getter = "HasMaxHeight", Setter = "SetHasMaxHeight", Category = "Layout", meta = (InlineEditConditionToggle, AllowPrivateAccess = "true"))
-	bool bHasMaxHeight;
+	UPROPERTY(EditAnywhere, Getter = "HasMaxHeight", Setter = "SetHasMaxHeight", Category = "Layout", meta = (InlineEditConditionToggle, AllowPrivateAccess = "true"))
+	bool bHasMaxHeight = false;
 
 	/** Dictates how to handle the text if it exceeds the max width */
-	UPROPERTY(EditAnywhere, BlueprintReadWrite, Getter, Setter, Category = "Layout", meta = (HideEditConditionToggle, EditCondition="bHasMaxWidth", AllowPrivateAccess = "true"))
+	UPROPERTY(EditAnywhere, Getter, Setter, Category = "Layout", meta = (HideEditConditionToggle, EditCondition="bHasMaxWidth", AllowPrivateAccess = "true"))
 	EText3DMaxWidthHandling MaxWidthHandling;
 
 	/** Sets a maximum height to the 3D Text */
-	UPROPERTY(EditAnywhere, BlueprintReadWrite, Getter, Setter, Category = "Layout", meta = (EditCondition = "bHasMaxHeight", ClampMin = 1, AllowPrivateAccess = "true"))
-	float MaxHeight;
+	UPROPERTY(EditAnywhere, Getter, Setter, Category = "Layout", meta = (EditCondition = "bHasMaxHeight", ClampMin = 1, AllowPrivateAccess = "true"))
+	float MaxHeight = 500.0f;
 
 	/** Should the mesh scale proportionally when Max Width/Height is set */
-	UPROPERTY(EditAnywhere, BlueprintReadWrite, Getter = "ScalesProportionally", Setter = "SetScaleProportionally", Category = "Layout", meta = (AllowPrivateAccess = "true"))
-	bool bScaleProportionally;
+	UPROPERTY(EditAnywhere, Getter = "ScalesProportionally", Setter = "SetScaleProportionally", Category = "Layout", meta = (AllowPrivateAccess = "true"))
+	bool bScaleProportionally = true;
 
 	// Lighting flags
 
 	/** Controls whether the text glyphs should cast a shadow or not. */
-	UPROPERTY(EditAnywhere, BlueprintReadWrite, Getter = "CastsShadow", Setter = "SetCastShadow", Category = "Lighting", meta = (AllowPrivateAccess = "true"))
+	UPROPERTY(EditAnywhere, Getter = "CastsShadow", Setter = "SetCastShadow", Category = "Lighting", meta = (AllowPrivateAccess = "true"))
 	bool bCastShadow = true;
 
-	/**
-	 * Returns the Text property, after being formatted by the FormatText virtual function.
-	 * If FormatText is not overriden, the return FText will be the same as the Text property.
-	 */
-	UFUNCTION(BlueprintCallable, Category = "Text3D")
-	FText GetFormattedText() const;
-
-	UFUNCTION()
-	TArray<FName> GetTypefaceNames() const;
-
-	/** Intercept and propagate a change on this component to all children. */
-	virtual void OnVisibilityChanged() override;
-
-	/** Intercept and propagate a change on this component to all children. */
-	virtual void OnHiddenInGameChanged() override;
-
-	/**
-	 * Will be called when text geometry is generated.
-	 * Override it to customize text formatting in the final geometry, without affecting the Text property.
-	 * Use GetFormattedText() to retrieve a FText with the result of this formatting.
-	 */
-	virtual void FormatText(FText& InOutText) const {}
-
-	/** Clears all generated components and meshes from this component. */
-	void ClearTextMesh();
-
-	/** Can be used to force an internal geometry and/or layout rebuild */
-	void TriggerInternalRebuild(const EText3DModifyFlags InModifyFlags);
-
 private:
+	/** Active renderer in charge of generating text geometry */
 	UPROPERTY()
-	TObjectPtr<USceneComponent> TextRoot;
+	TObjectPtr<UText3DRendererBase> TextRenderer;
 
 	UPROPERTY(BlueprintAssignable, Category = Events, meta = (AllowPrivateAccess = "true", DisplayName = "On Text Generated"))
 	FTextGenerated TextGeneratedDelegate;
@@ -476,19 +533,16 @@ private:
 	std::atomic<bool> bIsBuilding;
 
 	/** No mesh or layout rebuilds occur while this flag is true, allowing the user to set numerous properties before rebuilding. */
-	bool bFreezeBuild;
+	bool bFreezeBuild = false;
 
 	/** Used to determine and selectively perform the the type of rebuild requested. */
-	EText3DModifyFlags ModifyFlags;
+	EText3DModifyFlags ModifyFlags = EText3DModifyFlags::All;
 
 	/** Additional scale to apply to the text. */
-	FVector TextScale;
-
-	/** To uniquely identify and discard duplicate update requests. */
-	TArray<TSharedPtr<int32>> CachedCounterReferences;
+	FVector TextScale = FVector::ZeroVector;
 
 	/** Caches the last result of ShapedText, to allow faster updates of layout changes. */
-	TSharedPtr<struct FText3DShapedText> ShapedText;
+	TSharedPtr<FText3DShapedGlyphText> ShapedText;
 
 	/** Stores the text layout calculated by the TextLayoutMarshaller. */
 	TSharedPtr<FTextLayout> TextLayout;
@@ -496,28 +550,14 @@ private:
 	/** Determines how text is laid out, ie. parsing line breaks. */
 	TSharedPtr<ITextLayoutMarshaller> TextLayoutMarshaller;
 
-	/** Each character mesh is parented to a Kerning component. */
-	UPROPERTY(Transient)
-	TArray<TObjectPtr<USceneComponent>> CharacterKernings;
-
-	/** Each character mesh is held in a component. */
-	UPROPERTY(Transient)
-	TArray<TObjectPtr<UStaticMeshComponent>> CharacterMeshes;
-
 	/** Text statistics cached since last text generation */
 	FText3DStatistics Statistics;
 
-	/** Allocates, or shrinks existing components to match the input number. Returns false if nothing modified. */
-	bool AllocateGlyphs(int32 Num);
-
 	/** Slot based Material accessors. */
-	UMaterialInterface* GetMaterial(const EText3DGroupType Type) const;
 	void SetMaterial(const EText3DGroupType Type, UMaterialInterface* Material);
 	void OnMaterialChanged();
 
-	/** Validation and context sensitive limits applied to the current bevel value. */
-	void CheckBevel();
-	float MaxBevel() const;
+	float GetMaxBevel() const;
 
 	/** Selective build functionality (geometry, only layout, etc.) */
 	void UpdateTransforms();
@@ -528,16 +568,13 @@ private:
 	/** Layout functionality. */
 	float GetTextHeight() const;
 	void CalculateTextScale();
-	FVector GetLineLocation(int32 LineIndex);
 
 	/** Convenience functions to set and query ModifyFlags. */
 	bool NeedsMeshRebuild() const;
 	bool NeedsLayoutUpdate() const;
 	void MarkForGeometryUpdate();
-	void MarkForLayoutUpdate();
 	void ClearUpdateFlags();
 
-	uint32 GetTypeFaceIndex() const;
 	bool IsTypefaceAvailable(FName InTypeface) const;
 	TArray<FTypefaceEntry> GetAvailableTypefaces() const;
 	void RefreshTypeface();
diff --git a/Engine/Plugins/Experimental/Text3D/Source/Text3D/Public/Text3DTypes.h b/Engine/Plugins/Experimental/Text3D/Source/Text3D/Public/Text3DTypes.h
new file mode 100644
--- /dev/null
+++ b/Engine/Plugins/Experimental/Text3D/Source/Text3D/Public/Text3DTypes.h
@@ -0,0 +1,89 @@
+// Copyright Epic Games, Inc. All Rights Reserved.
+
+#pragma once
+
+#include "Text3DTypes.generated.h"
+
+UENUM()
+enum class EText3DBevelType : uint8
+{
+	Linear,
+	HalfCircle,
+	Convex,
+	Concave,
+	OneStep,
+	TwoSteps,
+	Engraved
+};
+
+enum class EText3DGroupType : uint8
+{
+	Front = 0,
+	Bevel = 1,
+	Extrude = 2,
+	Back = 3,
+
+	TypeCount = 4
+};
+
+UENUM()
+enum class EText3DVerticalTextAlignment : uint8
+{
+	FirstLine		UMETA(DisplayName = "First Line"),
+	Top				UMETA(DisplayName = "Top"),
+	Center			UMETA(DisplayName = "Center"),
+	Bottom			UMETA(DisplayName = "Bottom"),
+};
+
+UENUM()
+enum class EText3DHorizontalTextAlignment : uint8
+{
+	Left			UMETA(DisplayName = "Left"),
+	Center			UMETA(DisplayName = "Center"),
+	Right			UMETA(DisplayName = "Right"),
+};
+
+UENUM()
+enum class EText3DMaxWidthHandling : uint8
+{
+	/** Scales the text to meet the max width */
+	Scale			UMETA(DisplayName = "Scale"),
+	/** First wraps the text (if possible) and then scales to meet the max width */
+	WrapAndScale	UMETA(DisplayName = "Wrap and Scale"),
+};
+
+UENUM()
+enum class EText3DModifyFlags : uint8
+{
+	None = 0,
+	Layout = 1 << 0,
+	Geometry = 1 << 1,
+	Unfreeze = 1 << 2,
+
+	All = Layout | Geometry | Unfreeze
+};
+ENUM_CLASS_FLAGS(EText3DModifyFlags)
+
+UENUM()
+enum class EText3DCharacterEffectOrder : uint8
+{
+	Normal			UMETA(DisplayName = "Normal"),
+	FromCenter		UMETA(DisplayName = "From Center"),
+	ToCenter		UMETA(DisplayName = "To Center"),
+	Opposite		UMETA(DisplayName = "Opposite"),
+};
+
+struct FText3DWordStatistics
+{
+	/** Actual range taking into account whitespaces */
+	FTextRange ActualRange;
+
+	/** Render range not taking into account whitespaces */
+	FTextRange RenderRange;
+};
+
+struct FText3DStatistics
+{
+	TArray<FText3DWordStatistics> Words;
+	int32 WhiteSpaces;
+};
\ No newline at end of file
diff --git a/Engine/Plugins/Experimental/Text3D/Source/Text3DEditor/Private/ActorFactoryText3D.cpp b/Engine/Plugins/Experimental/Text3D/Source/Text3DEditor/Private/ActorFactoryText3D.cpp
deleted file mode 100644
diff --git a/Engine/Plugins/Experimental/Text3D/Source/Text3DEditor/Private/Factories/Text3DEditorActorFactory.cpp b/Engine/Plugins/Experimental/Text3D/Source/Text3DEditor/Private/Factories/Text3DEditorActorFactory.cpp
new file mode 100644
--- /dev/null
+++ b/Engine/Plugins/Experimental/Text3D/Source/Text3DEditor/Private/Factories/Text3DEditorActorFactory.cpp
@@ -0,0 +1,16 @@
+// Copyright Epic Games, Inc. All Rights Reserved.
+
+#include "Factories/Text3DEditorActorFactory.h"
+#include "Text3DActor.h"
+
+#define LOCTEXT_NAMESPACE "Text3DEditorActorFactory"
+
+UText3DEditorActorFactory::UText3DEditorActorFactory()
+{
+	DisplayName = LOCTEXT("Text3D", "Text 3D");
+	NewActorClass = AText3DActor::StaticClass();
+	SpawnPositionOffset = FVector(0, 0, 0);
+	bUseSurfaceOrientation = false;
+}
+
+#undef LOCTEXT_NAMESPACE
diff --git a/Engine/Plugins/Experimental/Text3D/Source/Text3DEditor/Private/Factories/Text3DEditorActorFactory.h b/Engine/Plugins/Experimental/Text3D/Source/Text3DEditor/Private/Factories/Text3DEditorActorFactory.h
new file mode 100644
--- /dev/null
+++ b/Engine/Plugins/Experimental/Text3D/Source/Text3DEditor/Private/Factories/Text3DEditorActorFactory.h
@@ -0,0 +1,15 @@
+// Copyright Epic Games, Inc. All Rights Reserved.
+
+#pragma once
+
+#include "ActorFactories/ActorFactory.h"
+#include "Text3DEditorActorFactory.generated.h"
+
+UCLASS()
+class UText3DEditorActorFactory : public UActorFactory
+{
+	GENERATED_BODY()
+
+public:
+	UText3DEditorActorFactory();
+};
diff --git a/Engine/Plugins/Experimental/Text3D/Source/Text3DEditor/Private/Styles/Text3DEditorStyle.cpp b/Engine/Plugins/Experimental/Text3D/Source/Text3DEditor/Private/Styles/Text3DEditorStyle.cpp
new file mode 100644
--- /dev/null
+++ b/Engine/Plugins/Experimental/Text3D/Source/Text3DEditor/Private/Styles/Text3DEditorStyle.cpp
@@ -0,0 +1,29 @@
+// Copyright Epic Games, Inc. All Rights Reserved.
+
+#include "Styles/Text3DEditorStyle.h"
+
+#include "Interfaces/IPluginManager.h"
+#include "Styling/SlateStyleMacros.h"
+#include "Styling/SlateStyleRegistry.h"
+
+FText3DEditorStyle::FText3DEditorStyle()
+	: FSlateStyleSet(UE_MODULE_NAME)
+{
+	const FVector2f Icon16x16(16.f, 16.f);
+	
+	const TSharedPtr<IPlugin> Plugin = IPluginManager::Get().FindPlugin(UE_PLUGIN_NAME);
+
+	check(Plugin.IsValid());
+
+	ContentRootDir = FPaths::Combine(Plugin->GetBaseDir(), TEXT("Resources"));
+
+	Set("ClassIcon.Text3DActor", new IMAGE_BRUSH("Text3DActor_16x", FVector2D (16.0f, 16.0f)));
+	Set("ClassThumbnail.Text3DActor", new IMAGE_BRUSH("Text3DActor_64x", FVector2D (64.0f, 64.0f)));
+
+	FSlateStyleRegistry::RegisterSlateStyle(*this);
+}
+
+FText3DEditorStyle::~FText3DEditorStyle()
+{
+	FSlateStyleRegistry::UnRegisterSlateStyle(*this);
+}
diff --git a/Engine/Plugins/Experimental/Text3D/Source/Text3DEditor/Private/Styles/Text3DEditorStyle.h b/Engine/Plugins/Experimental/Text3D/Source/Text3DEditor/Private/Styles/Text3DEditorStyle.h
new file mode 100644
--- /dev/null
+++ b/Engine/Plugins/Experimental/Text3D/Source/Text3DEditor/Private/Styles/Text3DEditorStyle.h
@@ -0,0 +1,19 @@
+// Copyright Epic Games, Inc. All Rights Reserved.
+
+#pragma once
+
+#include "Styling/SlateStyle.h"
+
+class FText3DEditorStyle : public FSlateStyleSet
+{
+public:
+	FText3DEditorStyle();
+	
+	virtual ~FText3DEditorStyle() override;
+	
+	static FText3DEditorStyle& Get()
+	{
+		static FText3DEditorStyle StyleSet;
+		return StyleSet;
+	}
+};
\ No newline at end of file
diff --git a/Engine/Plugins/Experimental/Text3D/Source/Text3DEditor/Private/Text3DEditor.cpp b/Engine/Plugins/Experimental/Text3D/Source/Text3DEditor/Private/Text3DEditor.cpp
deleted file mode 100644
diff --git a/Engine/Plugins/Experimental/Text3D/Source/Text3DEditor/Private/Text3DEditorModule.cpp b/Engine/Plugins/Experimental/Text3D/Source/Text3DEditor/Private/Text3DEditorModule.cpp
new file mode 100644
--- /dev/null
+++ b/Engine/Plugins/Experimental/Text3D/Source/Text3DEditor/Private/Text3DEditorModule.cpp
@@ -0,0 +1,13 @@
+// Copyright Epic Games, Inc. All Rights Reserved.
+
+#include "Text3DEditorModule.h"
+#include "Modules/ModuleManager.h"
+#include "Styles/Text3DEditorStyle.h"
+
+IMPLEMENT_MODULE(FText3DEditorModule, Text3DEditor)
+
+void FText3DEditorModule::StartupModule()
+{
+	// Init style set
+	FText3DEditorStyle::Get();
+}
diff --git a/Engine/Plugins/Experimental/Text3D/Source/Text3DEditor/Private/Text3DEditorModule.h b/Engine/Plugins/Experimental/Text3D/Source/Text3DEditor/Private/Text3DEditorModule.h
new file mode 100644
--- /dev/null
+++ b/Engine/Plugins/Experimental/Text3D/Source/Text3DEditor/Private/Text3DEditorModule.h
@@ -0,0 +1,14 @@
+// Copyright Epic Games, Inc. All Rights Reserved.
+#pragma once
+
+#include "Modules/ModuleInterface.h"
+
+class FText3DEditorModule final : public IModuleInterface
+{
+public:
+	FText3DEditorModule() = default;
+
+	//~ Begin IModuleInterface
+	virtual void StartupModule() override;
+	//~ End IModuleInterface
+};
\ No newline at end of file
diff --git a/Engine/Plugins/Experimental/Text3D/Source/Text3DEditor/Public/ActorFactoryText3D.h b/Engine/Plugins/Experimental/Text3D/Source/Text3DEditor/Public/ActorFactoryText3D.h
deleted file mode 100644
diff --git a/Engine/Plugins/Experimental/Text3D/Source/Text3DEditor/Public/Text3DEditor.h b/Engine/Plugins/Experimental/Text3D/Source/Text3DEditor/Public/Text3DEditor.h
deleted file mode 100644
