MotionDesign : Text3D
- Fix crash with text component when owner is not available (City sample)
- Updated logging

diff --git a/Engine/Plugins/Experimental/Text3D/Source/Text3D/Private/Logs/Text3DLogs.h b/Engine/Plugins/Experimental/Text3D/Source/Text3D/Private/Logs/Text3DLogs.h
new file mode 100644
--- /dev/null
+++ b/Engine/Plugins/Experimental/Text3D/Source/Text3D/Private/Logs/Text3DLogs.h
@@ -0,0 +1,7 @@
+// Copyright Epic Games, Inc. All Rights Reserved.
+
+#pragma once
+
+#include "Logging/LogMacros.h"
+
+DECLARE_LOG_CATEGORY_EXTERN(LogText3D, Log, All);
\ No newline at end of file
diff --git a/Engine/Plugins/Experimental/Text3D/Source/Text3D/Private/Renderers/StaticMeshes/Text3DStaticMeshesRenderer.cpp b/Engine/Plugins/Experimental/Text3D/Source/Text3D/Private/Renderers/StaticMeshes/Text3DStaticMeshesRenderer.cpp
--- a/Engine/Plugins/Experimental/Text3D/Source/Text3D/Private/Renderers/StaticMeshes/Text3DStaticMeshesRenderer.cpp
+++ b/Engine/Plugins/Experimental/Text3D/Source/Text3D/Private/Renderers/StaticMeshes/Text3DStaticMeshesRenderer.cpp
@@ -9,10 +9,11 @@
 #include "Engine/StaticMesh.h"
 #include "LayoutBuilders/Text3DShapedGlyphText.h"
 #include "LayoutBuilders/Text3DLayoutShaper.h"
+#include "Logs/Text3DLogs.h"
+#include "Misc/EnumerateRange.h"
 #include "Subsystems/Text3DEngineSubsystem.h"
 #include "Text3DComponent.h"
 #include "Text3DInternalTypes.h"
-#include "Misc/EnumerateRange.h"
 
 void UText3DStaticMeshesRenderer::OnCreate()
 {
diff --git a/Engine/Plugins/Experimental/Text3D/Source/Text3D/Private/Renderers/Text3DRendererBase.cpp b/Engine/Plugins/Experimental/Text3D/Source/Text3D/Private/Renderers/Text3DRendererBase.cpp
--- a/Engine/Plugins/Experimental/Text3D/Source/Text3D/Private/Renderers/Text3DRendererBase.cpp
+++ b/Engine/Plugins/Experimental/Text3D/Source/Text3D/Private/Renderers/Text3DRendererBase.cpp
@@ -2,6 +2,8 @@
 
 #include "Renderers/Text3DRendererBase.h"
 
+#include "GameFramework/Actor.h"
+#include "Logs/Text3DLogs.h"
 #include "Text3DComponent.h"
 
 UText3DComponent* UText3DRendererBase::GetText3DComponent() const
@@ -16,8 +18,15 @@ void UText3DRendererBase::Create()
 		return;
 	}
 
+	const UText3DComponent* Text3DComponent = GetText3DComponent();
+	if (!Text3DComponent || !Text3DComponent->IsRegistered())
+	{
+		return;
+	}
+
 	OnCreate();
 	bInitialized = true;
+	UE_LOG(LogText3D, Log, TEXT("%s : Text3DRenderer %s Created"), *Text3DComponent->GetOwner()->GetActorNameOrLabel(), *GetName().ToString())
 }
 
 void UText3DRendererBase::Update(EText3DRendererFlags InFlags)
@@ -27,6 +36,12 @@ void UText3DRendererBase::Update(EText3DRendererFlags InFlags)
 		return;
 	}
 
+	const UText3DComponent* Text3DComponent = GetText3DComponent();
+	if (!Text3DComponent || !Text3DComponent->IsRegistered())
+	{
+		return;
+	}
+
 	OnUpdate(InFlags);
 	
 	if (EnumHasAnyFlags(InFlags, EText3DRendererFlags::Geometry)
@@ -34,6 +49,8 @@ void UText3DRendererBase::Update(EText3DRendererFlags InFlags)
 	{
 		CachedBounds = OnCalculateBounds();
 	}
+
+	UE_LOG(LogText3D, Verbose, TEXT("%s : Text3DRenderer %s Updated with flags %i"), *Text3DComponent->GetOwner()->GetActorNameOrLabel(), *GetName().ToString(), InFlags)
 }
 
 void UText3DRendererBase::Clear()
@@ -43,8 +60,17 @@ void UText3DRendererBase::Clear()
 		return;
 	}
 
+	const UText3DComponent* Text3DComponent = GetText3DComponent();
+	if (!Text3DComponent)
+	{
+		return;
+	}
+
 	OnClear();
 	CachedBounds.Reset();
+
+	const AActor* Owner = Text3DComponent->GetOwner();
+	UE_LOG(LogText3D, Verbose, TEXT("%s : Text3DRenderer %s Cleared"), !!Owner ? *Owner->GetActorNameOrLabel() : TEXT("Invalid owner"), *GetName().ToString())
 }
 
 void UText3DRendererBase::Destroy()
@@ -54,9 +80,18 @@ void UText3DRendererBase::Destroy()
 		return;
 	}
 
+	const UText3DComponent* Text3DComponent = GetText3DComponent();
+	if (!Text3DComponent)
+	{
+		return;
+	}
+
 	OnDestroy();
 	CachedBounds.Reset();
 	bInitialized = false;
+
+	const AActor* Owner = Text3DComponent->GetOwner();
+	UE_LOG(LogText3D, Log, TEXT("%s : Text3DRenderer %s Destroyed"), !!Owner ? *Owner->GetActorNameOrLabel() : TEXT("Invalid owner"), *GetName().ToString())
 }
 
 FBox UText3DRendererBase::GetBounds() const
diff --git a/Engine/Plugins/Experimental/Text3D/Source/Text3D/Private/Subsystems/Text3DEngineSubsystem.cpp b/Engine/Plugins/Experimental/Text3D/Source/Text3D/Private/Subsystems/Text3DEngineSubsystem.cpp
--- a/Engine/Plugins/Experimental/Text3D/Source/Text3D/Private/Subsystems/Text3DEngineSubsystem.cpp
+++ b/Engine/Plugins/Experimental/Text3D/Source/Text3D/Private/Subsystems/Text3DEngineSubsystem.cpp
@@ -9,6 +9,7 @@
 #include "GeometryBuilders/Text3DGlyphContourNode.h"
 #include "GeometryBuilders/Text3DGlyphLoader.h"
 #include "GeometryBuilders/Text3DGlyphMeshBuilder.h"
+#include "Logs/Text3DLogs.h"
 #include "Materials/Material.h"
 #include "Misc/FileHelper.h"
 #include "UObject/ConstructorHelpers.h"
diff --git a/Engine/Plugins/Experimental/Text3D/Source/Text3D/Private/Text3DComponent.cpp b/Engine/Plugins/Experimental/Text3D/Source/Text3D/Private/Text3DComponent.cpp
--- a/Engine/Plugins/Experimental/Text3D/Source/Text3D/Private/Text3DComponent.cpp
+++ b/Engine/Plugins/Experimental/Text3D/Source/Text3D/Private/Text3DComponent.cpp
@@ -15,6 +15,7 @@
 #include "LayoutBuilders/Text3DLayout.h"
 #include "LayoutBuilders/Text3DShapedGlyphText.h"
 #include "LayoutBuilders/Text3DLayoutShaper.h"
+#include "Logs/Text3DLogs.h"
 #include "Materials/Material.h"
 #include "Misc/ScopeExit.h"
 #include "Misc/TransactionObjectEvent.h"
@@ -891,7 +892,7 @@ void UText3DComponent::TriggerInternalRebuild(const EText3DModifyFlags InModifyF
 
 void UText3DComponent::OnTextRendererClassChanged()
 {
-	UText3DRendererBase* RendererCDO = TextRendererClass.Get() ? TextRendererClass->GetDefaultObject<UText3DRendererBase>() : nullptr;
+	const UText3DRendererBase* RendererCDO = TextRendererClass.Get() ? TextRendererClass->GetDefaultObject<UText3DRendererBase>() : nullptr;
 
 	if (!RendererCDO)
 	{
@@ -922,7 +923,8 @@ void UText3DComponent::OnTextRendererClassChanged()
 
 	TextRenderer = NewObject<UText3DRendererBase>(this, TextRendererClass.Get(), RendererCDO->GetName());
 
-	UE_LOG(LogText3D, Log, TEXT("%s : Text3D renderer changed : old %s - new %s"), *GetOwner()->GetActorNameOrLabel(), *OldRendererName.ToString(), *TextRenderer->GetName().ToString());
+	const AActor* ContextActor = GetOwner();
+	UE_LOG(LogText3D, Log, TEXT("%s : Text3D renderer changed : old %s - new %s"), !!ContextActor ? *ContextActor->GetActorNameOrLabel() : TEXT("Invalid owner"), *OldRendererName.ToString(), *TextRenderer->GetName().ToString());
 
 	if (IsRegistered())
 	{
diff --git a/Engine/Plugins/Experimental/Text3D/Source/Text3D/Private/Text3DModule.cpp b/Engine/Plugins/Experimental/Text3D/Source/Text3D/Private/Text3DModule.cpp
--- a/Engine/Plugins/Experimental/Text3D/Source/Text3D/Private/Text3DModule.cpp
+++ b/Engine/Plugins/Experimental/Text3D/Source/Text3D/Private/Text3DModule.cpp
@@ -1,6 +1,8 @@
 // Copyright Epic Games, Inc. All Rights Reserved.
 
 #include "Text3DModule.h"
+
+#include "Logs/Text3DLogs.h"
 #include "Modules/ModuleManager.h"
 
 DEFINE_LOG_CATEGORY(LogText3D);
diff --git a/Engine/Plugins/Experimental/Text3D/Source/Text3D/Private/Text3DModule.h b/Engine/Plugins/Experimental/Text3D/Source/Text3D/Private/Text3DModule.h
--- a/Engine/Plugins/Experimental/Text3D/Source/Text3D/Private/Text3DModule.h
+++ b/Engine/Plugins/Experimental/Text3D/Source/Text3D/Private/Text3DModule.h
@@ -2,7 +2,6 @@
 
 #pragma once
 
-#include "Logging/LogMacros.h"
 #include "Modules/ModuleInterface.h"
 
 #if PLATFORM_COMPILER_HAS_GENERIC_KEYWORD
@@ -22,8 +21,6 @@ THIRD_PARTY_INCLUDES_END
 const int32 FontSize = 64;
 const float FontInverseScale = 1.0f / FontSize;
 
-DECLARE_LOG_CATEGORY_EXTERN(LogText3D, Log, All);
-
 class FText3DModule final : public IModuleInterface
 {
 public:
