MotionDesign : Text3D
- Fix regressions issue happening with AvaText3D caused by recent changes introduced to Text3D (crash, undo, property changes)

diff --git a/Engine/Plugins/Experimental/Avalanche/Source/AvalancheText/Private/AvaText3DComponent.cpp b/Engine/Plugins/Experimental/Avalanche/Source/AvalancheText/Private/AvaText3DComponent.cpp
--- a/Engine/Plugins/Experimental/Avalanche/Source/AvalancheText/Private/AvaText3DComponent.cpp
+++ b/Engine/Plugins/Experimental/Avalanche/Source/AvalancheText/Private/AvaText3DComponent.cpp
@@ -1151,7 +1151,7 @@ void UAvaText3DComponent::RegisterOnPropertyChangeFunctions()
 		PropertyChangeFunctions.Add(GET_MEMBER_NAME_CHECKED(UAvaText3DComponent, MotionDesignFont), &UAvaText3DComponent::RefreshGeometry);
 		PropertyChangeFunctions.Add(GET_MEMBER_NAME_CHECKED(UAvaText3DComponent, Font), &UAvaText3DComponent::RefreshMotionDesignFont);
 		PropertyChangeFunctions.Add(GET_MEMBER_NAME_CHECKED(UAvaText3DComponent, Text), &UAvaText3DComponent::RefreshGeometry);
-		PropertyChangeFunctions.Add(GET_MEMBER_NAME_CHECKED(UAvaText3DComponent, bEnforceUpperCase), &UAvaText3DComponent::RefreshFormatting);
+		PropertyChangeFunctions.Add(GET_MEMBER_NAME_CHECKED(UAvaText3DComponent, bEnforceUpperCase), &UAvaText3DComponent::OnTextChanged);
 
 		PropertyChangeFunctions.Add(GET_MEMBER_NAME_CHECKED(UAvaText3DComponent, ColoringStyle), &UAvaText3DComponent::RefreshMaterialInstances);
 		PropertyChangeFunctions.Add(GET_MEMBER_NAME_CHECKED(UAvaText3DComponent, MainTexture), &UAvaText3DComponent::RefreshMaterialInstances);
diff --git a/Engine/Plugins/Experimental/Text3D/Source/Text3D/Private/Renderers/StaticMeshes/Text3DStaticMeshesRenderer.cpp b/Engine/Plugins/Experimental/Text3D/Source/Text3D/Private/Renderers/StaticMeshes/Text3DStaticMeshesRenderer.cpp
--- a/Engine/Plugins/Experimental/Text3D/Source/Text3D/Private/Renderers/StaticMeshes/Text3DStaticMeshesRenderer.cpp
+++ b/Engine/Plugins/Experimental/Text3D/Source/Text3D/Private/Renderers/StaticMeshes/Text3DStaticMeshesRenderer.cpp
@@ -309,6 +309,11 @@ bool UText3DStaticMeshesRenderer::AllocateGlyphs(int32 Num)
 
 int32 UText3DStaticMeshesRenderer::GetGlyphCount()
 {
+	if (!TextRoot)
+	{
+		return 0;
+	}
+
 	return TextRoot->GetNumChildrenComponents();
 }
 
diff --git a/Engine/Plugins/Experimental/Text3D/Source/Text3D/Private/Renderers/Text3DRendererBase.cpp b/Engine/Plugins/Experimental/Text3D/Source/Text3D/Private/Renderers/Text3DRendererBase.cpp
--- a/Engine/Plugins/Experimental/Text3D/Source/Text3D/Private/Renderers/Text3DRendererBase.cpp
+++ b/Engine/Plugins/Experimental/Text3D/Source/Text3D/Private/Renderers/Text3DRendererBase.cpp
@@ -19,7 +19,7 @@ void UText3DRendererBase::Create()
 	}
 
 	const UText3DComponent* Text3DComponent = GetText3DComponent();
-	if (!Text3DComponent || !Text3DComponent->IsRegistered())
+	if (!IsValid(Text3DComponent))
 	{
 		return;
 	}
@@ -37,7 +37,7 @@ void UText3DRendererBase::Update(EText3DRendererFlags InFlags)
 	}
 
 	const UText3DComponent* Text3DComponent = GetText3DComponent();
-	if (!Text3DComponent || !Text3DComponent->IsRegistered())
+	if (!IsValid(Text3DComponent))
 	{
 		return;
 	}
diff --git a/Engine/Plugins/Experimental/Text3D/Source/Text3D/Private/Text3DComponent.cpp b/Engine/Plugins/Experimental/Text3D/Source/Text3D/Private/Text3DComponent.cpp
--- a/Engine/Plugins/Experimental/Text3D/Source/Text3D/Private/Text3DComponent.cpp
+++ b/Engine/Plugins/Experimental/Text3D/Source/Text3D/Private/Text3DComponent.cpp
@@ -68,6 +68,16 @@ void UText3DComponent::OnComponentCreated()
 	OnTextRendererClassChanged();
 }
 
+void UText3DComponent::OnComponentDestroyed(bool bInDestroyingHierarchy)
+{
+	Super::OnComponentDestroyed(bInDestroyingHierarchy);
+
+	if (TextRenderer)
+	{
+		TextRenderer->Destroy();
+	}
+}
+
 void UText3DComponent::PostLoad()
 {
 	Super::PostLoad();
@@ -78,16 +88,6 @@ void UText3DComponent::PostLoad()
 	OnTextRendererClassChanged();
 }
 
-void UText3DComponent::BeginDestroy()
-{
-	Super::BeginDestroy();
-
-	if (TextRenderer)
-	{
-		TextRenderer->Destroy();
-	}
-}
-
 bool UText3DComponent::NeedsMeshRebuild() const
 {
 	return EnumHasAnyFlags(ModifyFlags,EText3DModifyFlags::Geometry);
@@ -267,6 +267,11 @@ void UText3DComponent::PostEditChangeProperty(FPropertyChangedEvent& PropertyCha
 		GET_MEMBER_NAME_CHECKED(UText3DComponent, Typeface)
 	};
 
+	static const TSet<FName> TextPropertyNames
+	{
+		GET_MEMBER_NAME_CHECKED(UText3DComponent, Text)
+	};
+
 	static const TSet<FName> GeometryPropertyNames
 	{
 		GET_MEMBER_NAME_CHECKED(UText3DComponent, bOutline),
@@ -303,7 +308,7 @@ void UText3DComponent::PostEditChangeProperty(FPropertyChangedEvent& PropertyCha
 	{
 		OnTextRendererClassChanged();
 	}
-	if (MemberPropertyName == GET_MEMBER_NAME_CHECKED(UText3DComponent, Text))
+	if (TextPropertyNames.Contains(MemberPropertyName))
 	{
 		OnTextChanged();
 	}
diff --git a/Engine/Plugins/Experimental/Text3D/Source/Text3D/Public/Text3DComponent.h b/Engine/Plugins/Experimental/Text3D/Source/Text3D/Public/Text3DComponent.h
--- a/Engine/Plugins/Experimental/Text3D/Source/Text3D/Public/Text3DComponent.h
+++ b/Engine/Plugins/Experimental/Text3D/Source/Text3D/Public/Text3DComponent.h
@@ -347,12 +347,12 @@ public:
 protected:
 	//~ Begin UActorComponent
 	virtual void OnComponentCreated() override;
+	virtual void OnComponentDestroyed(bool bInDestroyingHierarchy) override;
 	virtual void OnRegister() override;
 	virtual void OnUnregister() override;
 	//~ End UActorComponent
 	
 	//~ Begin UObject
-	virtual void BeginDestroy() override;
 	virtual void PostLoad() override;
 	virtual void PostEditImport() override;
 	virtual void PostDuplicate(EDuplicateMode::Type InDuplicateMode) override;
@@ -415,6 +415,7 @@ protected:
 	FText Text;
 
 	/** Original text formatted */
+	UPROPERTY(Transient)
 	TOptional<FText> FormattedText;
 
 	/** Text font defines the style of rendered characters */
diff --git a/Engine/Plugins/Experimental/Text3D/Source/Text3DEditor/Text3DEditor.Build.cs b/Engine/Plugins/Experimental/Text3D/Source/Text3DEditor/Text3DEditor.Build.cs
--- a/Engine/Plugins/Experimental/Text3D/Source/Text3DEditor/Text3DEditor.Build.cs
+++ b/Engine/Plugins/Experimental/Text3D/Source/Text3DEditor/Text3DEditor.Build.cs
@@ -6,12 +6,13 @@ public class Text3DEditor : ModuleRules
 {
 	public Text3DEditor(ReadOnlyTargetRules Target) : base(Target)
 	{
-		PublicDependencyModuleNames.AddRange(new string[] {
+		PrivateDependencyModuleNames.AddRange(new string[] {
 			"Core",
 			"CoreUObject",
 			"EditorFramework",
 			"Projects",
 			"SlateCore",
+			"Slate",
 			"Text3D",
 			"UnrealEd",
 		});
