MotionDesign : Text3D

diff --git a/Engine/Plugins/Experimental/Text3D/Source/Text3D/Private/Text3DComponent.cpp b/Engine/Plugins/Experimental/Text3D/Source/Text3D/Private/Text3DComponent.cpp
--- a/Engine/Plugins/Experimental/Text3D/Source/Text3D/Private/Text3DComponent.cpp
+++ b/Engine/Plugins/Experimental/Text3D/Source/Text3D/Private/Text3DComponent.cpp
@@ -2,63 +2,55 @@
 
 #include "Text3DComponent.h"
 
-#include "Async/Async.h"
+#include "Containers/Ticker.h"
 #include "Engine/Engine.h"
 #include "Engine/Font.h"
 #include "Engine/StaticMesh.h"
+#include "Extensions/Text3DDefaultGeometryExtension.h"
+#include "Extensions/Text3DDefaultLayoutExtension.h"
+#include "Extensions/Text3DDefaultMaterialExtension.h"
+#include "Extensions/Text3DDefaultRenderingExtension.h"
 #include "Fonts/CompositeFont.h"
-#include "Fonts/SlateTextShaper.h"
-#include "Framework/Text/PlainTextLayoutMarshaller.h"
-#include "GeometryBuilders/Text3DGlyph.h"
-#include "GeometryBuilders/Text3DGlyphMeshBuilder.h"
 #include "Internationalization/Regex.h"
-#include "LayoutBuilders/Text3DLayout.h"
-#include "LayoutBuilders/Text3DShapedGlyphText.h"
-#include "LayoutBuilders/Text3DLayoutShaper.h"
 #include "Logs/Text3DLogs.h"
 #include "Materials/Material.h"
 #include "Misc/ScopeExit.h"
 #include "Misc/TransactionObjectEvent.h"
 #include "Renderers/Text3DRendererBase.h"
 #include "Renderers/StaticMeshes/Text3DStaticMeshesRenderer.h"
+#include "Settings/Text3DProjectSettings.h"
 #include "Styling/StyleDefaults.h"
 #include "Subsystems/Text3DEngineSubsystem.h"
-#include "Text3DModule.h"
+#include "Text3DComponentVersion.h"
 #include "UObject/ConstructorHelpers.h"
+#include "UObject/UObjectGlobals.h"
 
 #define LOCTEXT_NAMESPACE "Text3D"
 
+#if WITH_EDITOR
+UText3DComponent::FOnResolveFontByName UText3DComponent::OnResolveFontByNameDelegate;
+#endif
+
+FCustomVersionRegistration GRegisterText3DComponentVersion(FText3DComponentVersion::GUID, static_cast<int32>(FText3DComponentVersion::LatestVersion), TEXT("Text3DComponentVersion"));
+
 UText3DComponent::UText3DComponent()
-	: bIsBuilding(false)
-	, ShapedText(new FText3DShapedGlyphText())
+	: bIsUpdatingText(false)
 {
-	if (!IsRunningDedicatedServer())
-	{
-		struct FConstructorStatics
-		{
-			ConstructorHelpers::FObjectFinder<UFont> Font;
-			ConstructorHelpers::FObjectFinder<UMaterial> Material;
-			FConstructorStatics()
-				: Font(TEXT("/Engine/EngineFonts/Roboto"))
-				, Material(TEXT("/Engine/BasicShapes/BasicShapeMaterial"))
-			{
-			}
-		};
-		static FConstructorStatics ConstructorStatics;
+	const UText3DProjectSettings* Settings = UText3DProjectSettings::Get();
 
-		Font = ConstructorStatics.Font.Object;
-		RefreshTypeface();
+	Font = Settings->GetFallbackFont();
+	RefreshTypeface();
 
-		UMaterial* DefaultMaterial = ConstructorStatics.Material.Object;
-		FrontMaterial = DefaultMaterial;
-		BevelMaterial = DefaultMaterial;
-		ExtrudeMaterial = DefaultMaterial;
-		BackMaterial = DefaultMaterial;
-	}
+	Text = LOCTEXT("DefaultText", "Text");
 
 	// Legacy renderer
 	TextRendererClass = UText3DStaticMeshesRenderer::StaticClass();
-	Text = LOCTEXT("DefaultText", "Text");
+
+	// Extensions
+	MaterialExtension = CreateDefaultSubobject<UText3DMaterialExtensionBase, UText3DDefaultMaterialExtension>(TEXT("MaterialExtension"));
+	GeometryExtension = CreateDefaultSubobject<UText3DGeometryExtensionBase, UText3DDefaultGeometryExtension>(TEXT("GeometryExtension"));
+	LayoutExtension = CreateDefaultSubobject<UText3DLayoutExtensionBase, UText3DDefaultLayoutExtension>(TEXT("LayoutExtension"));
+	RenderingExtension = CreateDefaultSubobject<UText3DRenderingExtensionBase, UText3DDefaultRenderingExtension>(TEXT("RenderingExtension"));
 }
 
 void UText3DComponent::OnComponentCreated()
@@ -82,35 +74,34 @@ void UText3DComponent::PostLoad()
 {
 	Super::PostLoad();
 
-	// Reset so it's rebuilt (needed if re-using the component!)
-	ModifyFlags = EText3DModifyFlags::All;
-
 	OnTextRendererClassChanged();
 }
 
-bool UText3DComponent::NeedsMeshRebuild() const
+EText3DRendererFlags UText3DComponent::GetUpdateFlags() const
 {
-	return EnumHasAnyFlags(ModifyFlags,EText3DModifyFlags::Geometry);
+	return UpdateFlags;
 }
 
-bool UText3DComponent::NeedsLayoutUpdate() const
+void UText3DComponent::RequestUpdate(EText3DRendererFlags InFlags, bool bInImmediate)
 {
-	return EnumHasAnyFlags(ModifyFlags, EText3DModifyFlags::Layout);
-}
+	EnumAddFlags(UpdateFlags, InFlags);
 
-void UText3DComponent::MarkForGeometryUpdate()
-{
-	ModifyFlags |= EText3DModifyFlags::Geometry;
-}
-
-void UText3DComponent::MarkForLayoutUpdate()
-{
-	ModifyFlags |= EText3DModifyFlags::Layout;
+	if (UpdateFlags != EText3DRendererFlags::None)
+	{
+		if (bInImmediate)
+		{
+			RebuildInternal();
+		}
+		else
+		{
+			ScheduleTextUpdateNextTick();
+		}
+	}
 }
 
 void UText3DComponent::ClearUpdateFlags()
 {
-	ModifyFlags = EText3DModifyFlags::None;
+	UpdateFlags = EText3DRendererFlags::None;
 }
 
 uint32 UText3DComponent::GetTypeFaceIndex() const
@@ -224,7 +215,7 @@ void UText3DComponent::OnRegister()
 		TextRenderer->Create();
 	}
 	
-	RebuildInternal(/** AutoUpdate */true, /** CleanCache */true);
+	RequestUpdate(EText3DRendererFlags::None);
 }
 
 void UText3DComponent::OnUnregister()
@@ -240,18 +231,79 @@ void UText3DComponent::OnUnregister()
 	}
 }
 
+void UText3DComponent::Serialize(FArchive& InArchive)
+{
+	InArchive.UsingCustomVersion(FText3DComponentVersion::GUID);
+
+	Super::Serialize(InArchive);
+
+	const int32 Version = InArchive.CustomVer(FText3DComponentVersion::GUID);
+
+	if (Version < FText3DComponentVersion::LatestVersion)
+	{
+		UE_LOG(LogText3D, Log, TEXT("Text3D : Migrating from %i to %i version"), Version, FText3DComponentVersion::LatestVersion)
+
+		if (Version < FText3DComponentVersion::Extensions)
+		{
+#if WITH_EDITORONLY_DATA
+			PRAGMA_DISABLE_DEPRECATION_WARNINGS
+
+			if (MaterialExtension)
+			{
+				MaterialExtension->SetMaterial(EText3DGroupType::Front, FrontMaterial_DEPRECATED);
+				MaterialExtension->SetMaterial(EText3DGroupType::Back, BackMaterial_DEPRECATED);
+				MaterialExtension->SetMaterial(EText3DGroupType::Extrude, ExtrudeMaterial_DEPRECATED);
+				MaterialExtension->SetMaterial(EText3DGroupType::Bevel, BevelMaterial_DEPRECATED);
+			}
+
+			if (UText3DDefaultGeometryExtension* DefaultGeometryExtension = GetCastedGeometryExtension<UText3DDefaultGeometryExtension>())
+			{
+				DefaultGeometryExtension->SetExtrude(Extrude_DEPRECATED);
+				DefaultGeometryExtension->SetBevel(Bevel_DEPRECATED);
+				DefaultGeometryExtension->SetBevelType(BevelType_DEPRECATED);
+				DefaultGeometryExtension->SetBevelSegments(BevelSegments_DEPRECATED);
+				DefaultGeometryExtension->SetUseOutline(bOutline_DEPRECATED);
+				DefaultGeometryExtension->SetOutline(OutlineExpand_DEPRECATED);
+			}
+
+			if (UText3DDefaultLayoutExtension* DefaultLayoutExtension = GetCastedLayoutExtension<UText3DDefaultLayoutExtension>())
+			{
+				DefaultLayoutExtension->SetHorizontalAlignment(HorizontalAlignment_DEPRECATED);
+				DefaultLayoutExtension->SetVerticalAlignment(VerticalAlignment_DEPRECATED);
+				DefaultLayoutExtension->SetKerning(Kerning_DEPRECATED);
+				DefaultLayoutExtension->SetLineSpacing(LineSpacing_DEPRECATED);
+				DefaultLayoutExtension->SetWordSpacing(WordSpacing_DEPRECATED);
+				DefaultLayoutExtension->SetUseMaxWidth(bHasMaxWidth_DEPRECATED);
+				DefaultLayoutExtension->SetMaxWidth(MaxWidth_DEPRECATED);
+				DefaultLayoutExtension->SetMaxWidthBehavior(MaxWidthHandling_DEPRECATED);
+				DefaultLayoutExtension->SetUseMaxHeight(bHasMaxHeight_DEPRECATED);
+				DefaultLayoutExtension->SetMaxHeight(MaxHeight_DEPRECATED);
+				DefaultLayoutExtension->SetScaleProportionally(bScaleProportionally_DEPRECATED);
+			}
+			
+			if (RenderingExtension)
+			{
+				RenderingExtension->SetTextCastShadow(bCastShadow_DEPRECATED);
+			}
+
+			PRAGMA_ENABLE_DEPRECATION_WARNINGS
+#endif
+		}
+	}
+}
+
 void UText3DComponent::PostEditImport()
 {
 	Super::PostEditImport();
 
-	ModifyFlags = EText3DModifyFlags::All;
+	RequestUpdate(EText3DRendererFlags::All);
 }
 
 void UText3DComponent::PostDuplicate(EDuplicateMode::Type InDuplicateMode)
 {
 	Super::PostDuplicate(InDuplicateMode);
 
-	ModifyFlags = EText3DModifyFlags::All;
+	RequestUpdate(EText3DRendererFlags::All);
 }
 
 #if WITH_EDITOR
@@ -269,39 +321,8 @@ void UText3DComponent::PostEditChangeProperty(FPropertyChangedEvent& PropertyCha
 
 	static const TSet<FName> TextPropertyNames
 	{
-		GET_MEMBER_NAME_CHECKED(UText3DComponent, Text)
-	};
-
-	static const TSet<FName> GeometryPropertyNames
-	{
-		GET_MEMBER_NAME_CHECKED(UText3DComponent, bOutline),
-		GET_MEMBER_NAME_CHECKED(UText3DComponent, OutlineExpand),
-		GET_MEMBER_NAME_CHECKED(UText3DComponent, Extrude),
-		GET_MEMBER_NAME_CHECKED(UText3DComponent, Bevel),
-		GET_MEMBER_NAME_CHECKED(UText3DComponent, BevelType),
-		GET_MEMBER_NAME_CHECKED(UText3DComponent, BevelSegments),
-	};
-
-	static const TSet<FName> LayoutPropertyNames
-	{
-		GET_MEMBER_NAME_CHECKED(UText3DComponent, Kerning),
-		GET_MEMBER_NAME_CHECKED(UText3DComponent, LineSpacing),
-		GET_MEMBER_NAME_CHECKED(UText3DComponent, WordSpacing),
-		GET_MEMBER_NAME_CHECKED(UText3DComponent, HorizontalAlignment),
-		GET_MEMBER_NAME_CHECKED(UText3DComponent, VerticalAlignment),
-		GET_MEMBER_NAME_CHECKED(UText3DComponent, bHasMaxWidth),
-		GET_MEMBER_NAME_CHECKED(UText3DComponent, MaxWidth),
-		GET_MEMBER_NAME_CHECKED(UText3DComponent, bHasMaxHeight),
-		GET_MEMBER_NAME_CHECKED(UText3DComponent, MaxHeight),
-		GET_MEMBER_NAME_CHECKED(UText3DComponent, bScaleProportionally),
-	};
-
-	static const TSet<FName> MaterialPropertyNames
-	{
-		GET_MEMBER_NAME_CHECKED(UText3DComponent, FrontMaterial),
-		GET_MEMBER_NAME_CHECKED(UText3DComponent, BevelMaterial), 
-		GET_MEMBER_NAME_CHECKED(UText3DComponent, ExtrudeMaterial),
-		GET_MEMBER_NAME_CHECKED(UText3DComponent, BackMaterial),  
+		GET_MEMBER_NAME_CHECKED(UText3DComponent, Text),
+		GET_MEMBER_NAME_CHECKED(UText3DComponent, bEnforceUpperCase)
 	};
 
 	if (MemberPropertyName == GET_MEMBER_NAME_CHECKED(UText3DComponent, TextRendererClass))
@@ -316,25 +337,11 @@ void UText3DComponent::PostEditChangeProperty(FPropertyChangedEvent& PropertyCha
 	{
 		OnFontPropertiesChanged();
 	}
-	else if (GeometryPropertyNames.Contains(MemberPropertyName))
-	{
-		OnGeometryPropertiesChanged();
-	}
-	else if (LayoutPropertyNames.Contains(MemberPropertyName))
-	{
-		OnLayoutPropertiesChanged();
-	}
-	else if (MemberPropertyName == GET_MEMBER_NAME_CHECKED(UText3DComponent, MaxWidthHandling))
-	{
-		OnLayoutHandlingPropertiesChanged();
-	}
-	else if (MaterialPropertyNames.Contains(MemberPropertyName))
-	{
-		OnMaterialChanged();
-	}
-	else if (MemberPropertyName == GET_MEMBER_NAME_CHECKED(UText3DComponent, bCastShadow))
+	else if (MemberPropertyName == GET_MEMBER_NAME_CHECKED(UText3DComponent, TransformExtensions))
 	{
-		OnVisibilityPropertiesChanged();
+		EText3DRendererFlags Flags = EText3DRendererFlags::Layout;
+		EnumAddFlags(Flags, EText3DRendererFlags::Material);
+		RequestUpdate(Flags);
 	}
 }
 
@@ -344,19 +351,26 @@ void UText3DComponent::PostTransacted(const FTransactionObjectEvent& Transaction
 
 	if (TransactionEvent.GetEventType() == ETransactionObjectEventType::UndoRedo)
 	{
-		ModifyFlags |= EText3DModifyFlags::All;
-		RebuildInternal();
+		RequestUpdate(EText3DRendererFlags::All);
 	}
 }
-#endif
 
-void UText3DComponent::SetRefreshOnChange(const bool Value)
+void UText3DComponent::SetFontByName(const FString& InFontName)
 {
-	if (bRefreshOnChange != Value)
+	if (OnResolveFontByNameDelegate.IsBound())
 	{
-		bRefreshOnChange = Value;
+		if (UFont* ResolvedFont = OnResolveFontByNameDelegate.Execute(InFontName))
+		{
+			SetFont(ResolvedFont);
+		}
 	}
 }
+#endif
+
+const FText& UText3DComponent::GetFormattedText() const
+{
+	return FormattedText.Get(FText::GetEmpty());
+}
 
 void UText3DComponent::SetText(const FText& Value)
 {
@@ -369,6 +383,17 @@ void UText3DComponent::SetText(const FText& Value)
 	OnTextChanged();
 }
 
+void UText3DComponent::SetEnforceUpperCase(bool bInEnforceUpperCase)
+{
+	if (bEnforceUpperCase == bInEnforceUpperCase)
+	{
+		return;
+	}
+
+	bEnforceUpperCase = bInEnforceUpperCase;
+	OnTextChanged();
+}
+
 FText UText3DComponent::GetFormattedText()
 {
 	if (!FormattedText.IsSet())
@@ -393,133 +418,132 @@ void UText3DComponent::SetFont(UFont* InFont)
 
 void UText3DComponent::SetHasOutline(const bool bValue)
 {
-	if (bOutline == bValue)
+	if (UText3DDefaultGeometryExtension* DefaultGeometryExtension = GetCastedGeometryExtension<UText3DDefaultGeometryExtension>())
 	{
-		return;
+		DefaultGeometryExtension->SetUseOutline(bValue);
 	}
-
-	bOutline = bValue;
-	OnGeometryPropertiesChanged();
 }
 
 void UText3DComponent::SetOutlineExpand(const float Value)
 {
-	if (FMath::IsNearlyEqual(OutlineExpand, Value))
+	if (UText3DDefaultGeometryExtension* DefaultGeometryExtension = GetCastedGeometryExtension<UText3DDefaultGeometryExtension>())
 	{
-		return;
+		DefaultGeometryExtension->SetOutline(Value);
 	}
-
-	OutlineExpand = Value;
-	OnGeometryPropertiesChanged();
 }
 
 void UText3DComponent::SetExtrude(const float Value)
 {
-	const float NewValue = FMath::Max(0.0f, Value);
-	if (FMath::IsNearlyEqual(Extrude, NewValue))
+	if (UText3DDefaultGeometryExtension* DefaultGeometryExtension = GetCastedGeometryExtension<UText3DDefaultGeometryExtension>())
 	{
-		return;
+		DefaultGeometryExtension->SetExtrude(Value);
 	}
-
-	Extrude = NewValue;
-	OnGeometryPropertiesChanged();
 }
 
 void UText3DComponent::SetBevel(const float Value)
 {
-	const float NewValue = FMath::Clamp(Value, 0.f, GetMaxBevel());
-	if (FMath::IsNearlyEqual(Bevel, NewValue))
+	if (UText3DDefaultGeometryExtension* DefaultGeometryExtension = GetCastedGeometryExtension<UText3DDefaultGeometryExtension>())
 	{
-		return;
+		DefaultGeometryExtension->SetBevel(Value);
 	}
-	
-	Bevel = NewValue;
-	OnGeometryPropertiesChanged();
 }
 
 void UText3DComponent::SetBevelType(const EText3DBevelType Value)
 {
-	if (BevelType == Value)
+	if (UText3DDefaultGeometryExtension* DefaultGeometryExtension = GetCastedGeometryExtension<UText3DDefaultGeometryExtension>())
 	{
-		return;
+		DefaultGeometryExtension->SetBevelType(Value);
 	}
-
-	BevelType = Value;
-	OnGeometryPropertiesChanged();
 }
 
 void UText3DComponent::SetBevelSegments(const int32 Value)
 {
-	int32 MinBevelSegments = 1;
-	if (BevelType == EText3DBevelType::HalfCircle)
+	if (UText3DDefaultGeometryExtension* DefaultGeometryExtension = GetCastedGeometryExtension<UText3DDefaultGeometryExtension>())
 	{
-		MinBevelSegments = 2;
+		DefaultGeometryExtension->SetBevelSegments(Value);
 	}
+}
 
-	const int32 NewValue = FMath::Clamp(Value, MinBevelSegments, 15);
-	if (BevelSegments == NewValue)
+UMaterialInterface* UText3DComponent::GetBackMaterial() const
+{
+	if (MaterialExtension)
 	{
-		return;
+		return MaterialExtension->GetMaterial(EText3DGroupType::Back);
 	}
-	
-	BevelSegments = NewValue;
-	OnGeometryPropertiesChanged();
-}
 
-void UText3DComponent::SetFrontMaterial(UMaterialInterface* Value)
-{
-	SetMaterial(EText3DGroupType::Front, Value);
+	return nullptr;
 }
 
-void UText3DComponent::SetBevelMaterial(UMaterialInterface* Value)
+UMaterialInterface* UText3DComponent::GetExtrudeMaterial() const
 {
-	SetMaterial(EText3DGroupType::Bevel, Value);
-}
+	if (MaterialExtension)
+	{
+		return MaterialExtension->GetMaterial(EText3DGroupType::Extrude);
+	}
 
-void UText3DComponent::SetExtrudeMaterial(UMaterialInterface* Value)
-{
-	SetMaterial(EText3DGroupType::Extrude, Value);
+	return nullptr;
 }
 
-void UText3DComponent::SetBackMaterial(UMaterialInterface* Value)
+UMaterialInterface* UText3DComponent::GetBevelMaterial() const
 {
-	SetMaterial(EText3DGroupType::Back, Value);
+	if (MaterialExtension)
+	{
+		return MaterialExtension->GetMaterial(EText3DGroupType::Bevel);
+	}
+
+	return nullptr;
 }
 
-UMaterialInterface* UText3DComponent::GetMaterial(const EText3DGroupType Type) const
+UMaterialInterface* UText3DComponent::GetFrontMaterial() const
 {
-	switch (Type)
-	{
-	case EText3DGroupType::Front:
+	if (MaterialExtension)
 	{
-		return FrontMaterial;
+		return MaterialExtension->GetMaterial(EText3DGroupType::Front);
 	}
 
-	case EText3DGroupType::Bevel:
+	return nullptr;
+}
+
+void UText3DComponent::SetFrontMaterial(UMaterialInterface* Value)
+{
+	if (MaterialExtension)
 	{
-		return BevelMaterial;
+		return MaterialExtension->SetMaterial(EText3DGroupType::Front, Value);
 	}
+}
 
-	case EText3DGroupType::Extrude:
+void UText3DComponent::SetBevelMaterial(UMaterialInterface* Value)
+{
+	if (MaterialExtension)
 	{
-		return ExtrudeMaterial;
+		return MaterialExtension->SetMaterial(EText3DGroupType::Bevel, Value);
 	}
+}
 
-	case EText3DGroupType::Back:
+void UText3DComponent::SetExtrudeMaterial(UMaterialInterface* Value)
+{
+	if (MaterialExtension)
 	{
-		return BackMaterial;
+		return MaterialExtension->SetMaterial(EText3DGroupType::Extrude, Value);
 	}
+}
 
-	default:
+void UText3DComponent::SetBackMaterial(UMaterialInterface* Value)
+{
+	if (MaterialExtension)
 	{
-		return nullptr;
-	}
+		return MaterialExtension->SetMaterial(EText3DGroupType::Back, Value);
 	}
 }
 
-TSharedPtr<FText3DShapedGlyphText> UText3DComponent::GetShapedText() const
+float UText3DComponent::GetKerning() const
 {
-	return ShapedText;
+	if (const UText3DDefaultLayoutExtension* DefaultLayoutExtension = GetCastedLayoutExtension<UText3DDefaultLayoutExtension>())
+	{
+		return DefaultLayoutExtension->GetKerning();
+	}
+
+	return 0.f;
 }
 
 UText3DRendererBase* UText3DComponent::GetTextRenderer() const
@@ -527,372 +551,320 @@ UText3DRendererBase* UText3DComponent::GetTextRenderer() const
 	return TextRenderer;
 }
 
-void UText3DComponent::SetMaterial(const EText3DGroupType Type, UMaterialInterface* Value)
+void UText3DComponent::SetKerning(const float Value)
 {
-	UMaterialInterface* OldMaterial = GetMaterial(Type);
-	if (Value != OldMaterial)
+	if (UText3DDefaultLayoutExtension* DefaultLayoutExtension = GetCastedLayoutExtension<UText3DDefaultLayoutExtension>())
 	{
-		switch (Type)
-		{
-		case EText3DGroupType::Front:
-		{
-			FrontMaterial = Value;
-			break;
-		}
-
-		case EText3DGroupType::Back:
-		{
-			BackMaterial = Value;
-			break;
-		}
-
-		case EText3DGroupType::Extrude:
-		{
-			ExtrudeMaterial = Value;
-			break;
-		}
-
-		case EText3DGroupType::Bevel:
-		{
-			BevelMaterial = Value;
-			break;
-		}
-
-		default:
-		{
-			return;
-		}
-		}
-
-		OnMaterialChanged();
+		DefaultLayoutExtension->SetKerning(Value);
 	}
 }
 
-void UText3DComponent::SetKerning(const float Value)
+float UText3DComponent::GetLineSpacing() const
 {
-	if (FMath::IsNearlyEqual(Kerning, Value))
+	if (const UText3DDefaultLayoutExtension* DefaultLayoutExtension = GetCastedLayoutExtension<UText3DDefaultLayoutExtension>())
 	{
-		return;
+		return DefaultLayoutExtension->GetLineSpacing();
 	}
-	
-	Kerning = Value;
-	OnLayoutPropertiesChanged();
+
+	return 0.f;
 }
 
 void UText3DComponent::SetLineSpacing(const float Value)
 {
-	if (FMath::IsNearlyEqual(LineSpacing, Value))
+	if (UText3DDefaultLayoutExtension* DefaultLayoutExtension = GetCastedLayoutExtension<UText3DDefaultLayoutExtension>())
 	{
-		return;
+		DefaultLayoutExtension->SetLineSpacing(Value);
 	}
-	
-	LineSpacing = Value;
-	OnLayoutPropertiesChanged();
 }
 
-void UText3DComponent::SetWordSpacing(const float Value)
+float UText3DComponent::GetWordSpacing() const
 {
-	if (FMath::IsNearlyEqual(WordSpacing, Value))
+	if (const UText3DDefaultLayoutExtension* DefaultLayoutExtension = GetCastedLayoutExtension<UText3DDefaultLayoutExtension>())
 	{
-		return;
+		return DefaultLayoutExtension->GetWordSpacing();
 	}
 
-	WordSpacing = Value;
-	OnLayoutPropertiesChanged();
+	return 0.f;
 }
 
-void UText3DComponent::SetHorizontalAlignment(const EText3DHorizontalTextAlignment Value)
+void UText3DComponent::SetWordSpacing(const float Value)
 {
-	if (HorizontalAlignment == Value)
+	if (UText3DDefaultLayoutExtension* DefaultLayoutExtension = GetCastedLayoutExtension<UText3DDefaultLayoutExtension>())
 	{
-		return;
+		DefaultLayoutExtension->SetWordSpacing(Value);
 	}
+}
 
-	HorizontalAlignment = Value;
-	OnLayoutPropertiesChanged();
+EText3DHorizontalTextAlignment UText3DComponent::GetHorizontalAlignment() const
+{
+	if (const UText3DDefaultLayoutExtension* DefaultLayoutExtension = GetCastedLayoutExtension<UText3DDefaultLayoutExtension>())
+	{
+		return DefaultLayoutExtension->GetHorizontalAlignment();
+	}
+
+	return EText3DHorizontalTextAlignment::Left;
 }
 
-void UText3DComponent::SetVerticalAlignment(const EText3DVerticalTextAlignment Value)
+void UText3DComponent::SetHorizontalAlignment(const EText3DHorizontalTextAlignment Value)
 {
-	if (VerticalAlignment == Value)
+	if (UText3DDefaultLayoutExtension* DefaultLayoutExtension = GetCastedLayoutExtension<UText3DDefaultLayoutExtension>())
 	{
-		return;
+		DefaultLayoutExtension->SetHorizontalAlignment(Value);
 	}
-	
-	VerticalAlignment = Value;
-	OnLayoutPropertiesChanged();
 }
 
-void UText3DComponent::SetHasMaxWidth(const bool Value)
+EText3DVerticalTextAlignment UText3DComponent::GetVerticalAlignment() const
 {
-	if (bHasMaxWidth == Value)
+	if (const UText3DDefaultLayoutExtension* DefaultLayoutExtension = GetCastedLayoutExtension<UText3DDefaultLayoutExtension>())
 	{
-		return;
+		return DefaultLayoutExtension->GetVerticalAlignment();
 	}
 
-	bHasMaxWidth = Value;
-	OnLayoutPropertiesChanged();
+	return EText3DVerticalTextAlignment::FirstLine;
 }
 
-void UText3DComponent::SetMaxWidth(const float Value)
+void UText3DComponent::SetVerticalAlignment(const EText3DVerticalTextAlignment Value)
 {
-	const float NewValue = FMath::Max(1.0f, Value);
-	if (FMath::IsNearlyEqual(MaxWidth, NewValue))
+	if (UText3DDefaultLayoutExtension* DefaultLayoutExtension = GetCastedLayoutExtension<UText3DDefaultLayoutExtension>())
 	{
-		return;
+		DefaultLayoutExtension->SetVerticalAlignment(Value);
 	}
-	
-	MaxWidth = NewValue;
-	OnLayoutPropertiesChanged();
 }
 
-void UText3DComponent::SetMaxWidthHandling(const EText3DMaxWidthHandling Value)
+bool UText3DComponent::HasMaxWidth() const
 {
-	if (MaxWidthHandling == Value)
+	if (const UText3DDefaultLayoutExtension* DefaultLayoutExtension = GetCastedLayoutExtension<UText3DDefaultLayoutExtension>())
 	{
-		return;
+		return DefaultLayoutExtension->GetUseMaxWidth();
 	}
 
-	MaxWidthHandling = Value;
-	OnLayoutHandlingPropertiesChanged();
+	return false;
 }
 
-void UText3DComponent::SetHasMaxHeight(const bool Value)
+void UText3DComponent::SetHasMaxWidth(const bool Value)
 {
-	if (bHasMaxHeight == Value)
+	if (UText3DDefaultLayoutExtension* DefaultLayoutExtension = GetCastedLayoutExtension<UText3DDefaultLayoutExtension>())
 	{
-		return;
+		DefaultLayoutExtension->SetUseMaxWidth(Value);
 	}
-	
-	bHasMaxHeight = Value;
-	OnLayoutPropertiesChanged();
 }
 
-void UText3DComponent::SetMaxHeight(const float Value)
+float UText3DComponent::GetMaxWidth() const
 {
-	const float NewValue = FMath::Max(1.0f, Value);
-	if (FMath::IsNearlyEqual(MaxHeight, NewValue))
+	if (const UText3DDefaultLayoutExtension* DefaultLayoutExtension = GetCastedLayoutExtension<UText3DDefaultLayoutExtension>())
 	{
-		return;
+		return DefaultLayoutExtension->GetMaxWidth();
 	}
 
-	MaxHeight = NewValue;
-	OnLayoutPropertiesChanged();
+	return 0.f;
 }
 
-void UText3DComponent::SetScaleProportionally(const bool Value)
+void UText3DComponent::SetMaxWidth(const float Value)
 {
-	if (bScaleProportionally == Value)
+	if (UText3DDefaultLayoutExtension* DefaultLayoutExtension = GetCastedLayoutExtension<UText3DDefaultLayoutExtension>())
 	{
-		return;
+		DefaultLayoutExtension->SetMaxWidth(Value);
 	}
-
-	bScaleProportionally = Value;
-	OnLayoutPropertiesChanged();
 }
 
-void UText3DComponent::SetFreeze(const bool bFreeze)
+EText3DMaxWidthHandling UText3DComponent::GetMaxWidthHandling() const
 {
-	bFreezeBuild = bFreeze;
-	if (bFreeze)
+	if (const UText3DDefaultLayoutExtension* DefaultLayoutExtension = GetCastedLayoutExtension<UText3DDefaultLayoutExtension>())
 	{
-		ModifyFlags |= EText3DModifyFlags::Unfreeze;
+		return DefaultLayoutExtension->GetMaxWidthBehavior();
 	}
-	else if (EnumHasAnyFlags(ModifyFlags,EText3DModifyFlags::Unfreeze))
+
+	return EText3DMaxWidthHandling::Scale;
+}
+
+void UText3DComponent::SetMaxWidthHandling(const EText3DMaxWidthHandling Value)
+{
+	if (UText3DDefaultLayoutExtension* DefaultLayoutExtension = GetCastedLayoutExtension<UText3DDefaultLayoutExtension>())
 	{
-		RebuildInternal();
+		DefaultLayoutExtension->SetMaxWidthBehavior(Value);
 	}
 }
 
-void UText3DComponent::SetCastShadow(bool NewCastShadow)
+bool UText3DComponent::HasMaxHeight() const
 {
-	if (NewCastShadow == bCastShadow)
+	if (const UText3DDefaultLayoutExtension* DefaultLayoutExtension = GetCastedLayoutExtension<UText3DDefaultLayoutExtension>())
 	{
-		return;
+		return DefaultLayoutExtension->GetUseMaxHeight();
 	}
 
-	bCastShadow = NewCastShadow;
-	OnVisibilityPropertiesChanged();
+	return false;
 }
 
-void UText3DComponent::SetTypeface(const FName InTypeface)
+void UText3DComponent::SetHasMaxHeight(const bool Value)
 {
-	if (Typeface == InTypeface || !IsTypefaceAvailable(InTypeface))
+	if (UText3DDefaultLayoutExtension* DefaultLayoutExtension = GetCastedLayoutExtension<UText3DDefaultLayoutExtension>())
 	{
-		return;
+		DefaultLayoutExtension->SetUseMaxHeight(Value);
 	}
-
-	Typeface = InTypeface;
-	OnFontPropertiesChanged();
 }
 
-void UText3DComponent::SetTextRendererClass(const TSubclassOf<UText3DRendererBase>& InClass)
+float UText3DComponent::GetMaxHeight() const
 {
-	if (!TextRendererClass.Get()
-		|| TextRendererClass == InClass)
+	if (const UText3DDefaultLayoutExtension* DefaultLayoutExtension = GetCastedLayoutExtension<UText3DDefaultLayoutExtension>())
 	{
-		return;
+		return DefaultLayoutExtension->GetMaxHeight();
 	}
 
-	TextRendererClass = InClass;
-	OnTextRendererClassChanged();
+	return 0.f;
 }
 
-void UText3DComponent::Rebuild()
+void UText3DComponent::SetMaxHeight(const float Value)
 {
-	// Rebuild, ignoring RefreshOnChange
-	RebuildInternal(false);
+	if (UText3DDefaultLayoutExtension* DefaultLayoutExtension = GetCastedLayoutExtension<UText3DDefaultLayoutExtension>())
+	{
+		DefaultLayoutExtension->SetMaxHeight(Value);
+	}
 }
 
-void UText3DComponent::RebuildInternal(const bool& bIsAutoUpdate, const bool& bCleanCache)
+bool UText3DComponent::ScalesProportionally() const
 {
-	// If this is an auto update, but the flag is off, ignore this rebuild request
-	if (bIsAutoUpdate && bRefreshOnChange == false)
+	if (const UText3DDefaultLayoutExtension* DefaultLayoutExtension = GetCastedLayoutExtension<UText3DDefaultLayoutExtension>())
 	{
-		return;
+		return DefaultLayoutExtension->GetScalesProportionally();
 	}
 
-	if (NeedsMeshRebuild() && !bFreezeBuild)
-	{
-		BuildTextMesh(bCleanCache);
-	}
-	else if (NeedsLayoutUpdate())
-	{
-		UpdateTransforms();
-	}
+	return false;
 }
 
-float UText3DComponent::GetTextHeight() const
+void UText3DComponent::SetScaleProportionally(const bool Value)
 {
-	return ShapedText->Lines.Num() * ShapedText->LineHeight + (ShapedText->Lines.Num() - 1) * LineSpacing;
+	if (UText3DDefaultLayoutExtension* DefaultLayoutExtension = GetCastedLayoutExtension<UText3DDefaultLayoutExtension>())
+	{
+		DefaultLayoutExtension->SetScaleProportionally(Value);
+	}
 }
 
-void UText3DComponent::CalculateTextScale()
+bool UText3DComponent::CastsShadow() const
 {
-	FVector Scale(1.0f, 1.0f, 1.0f);
-
-	float TextMaxWidth = 0.0f;
-	for (const FShapedGlyphLine& ShapedLine : ShapedText->Lines)
+	if (RenderingExtension)
 	{
-		TextMaxWidth = FMath::Max(TextMaxWidth, ShapedLine.Width);
+		return RenderingExtension->GetTextCastShadow();
 	}
 
-	if (bHasMaxWidth && TextMaxWidth > MaxWidth && TextMaxWidth > 0.0f)
-	{
-		Scale.Y *= MaxWidth / TextMaxWidth;
-		if (bScaleProportionally)
-		{
-			Scale.Z = Scale.Y;
-		}
-	}
+	return false;
+}
 
-	const float TotalHeight = GetTextHeight();
-	if (bHasMaxHeight && TotalHeight > MaxHeight && TotalHeight > 0.0f)
+void UText3DComponent::SetCastShadow(bool NewCastShadow)
+{
+	if (RenderingExtension)
 	{
-		Scale.Z *= MaxHeight / TotalHeight;
-		if (bScaleProportionally)
-		{
-			Scale.Y = Scale.Z;
-		}
+		RenderingExtension->SetTextCastShadow(NewCastShadow);
 	}
+}
 
-	if (bScaleProportionally)
+void UText3DComponent::SetTypeface(const FName InTypeface)
+{
+	if (Typeface == InTypeface || !IsTypefaceAvailable(InTypeface))
 	{
-		Scale.X = Scale.Y;
+		return;
 	}
 
-	TextScale = Scale;
+	Typeface = InTypeface;
+	OnFontPropertiesChanged();
 }
 
-FVector UText3DComponent::GetTextScale()
+void UText3DComponent::SetTextRendererClass(const TSubclassOf<UText3DRendererBase>& InClass)
 {
-	if (TextScale == FVector::ZeroVector)
+	if (!TextRendererClass.Get()
+		|| TextRendererClass == InClass)
 	{
-		CalculateTextScale();
+		return;
 	}
 
-	return TextScale;
+	TextRendererClass = InClass;
+	OnTextRendererClassChanged();
 }
 
-FVector UText3DComponent::GetLineLocation(int32 LineIndex) const
+void UText3DComponent::RebuildInternal(bool bCleanCache)
 {
-	float HorizontalOffset = 0.0f, VerticalOffset = 0.0f;
-	if (LineIndex < 0 || LineIndex >= ShapedText->Lines.Num())
+	if (!IsRegistered())
 	{
-		return FVector();
+		return;
 	}
-
-	const FShapedGlyphLine& ShapedLine = ShapedText->Lines[LineIndex];
-
-	if (HorizontalAlignment == EText3DHorizontalTextAlignment::Center)
+	
+	if (UE::IsSavingPackage(this))
 	{
-		HorizontalOffset = -ShapedLine.Width * 0.5f;
+		return;
 	}
-	else if (HorizontalAlignment == EText3DHorizontalTextAlignment::Right)
+	
+	if (!TextRenderer)
 	{
-		HorizontalOffset = -ShapedLine.Width;
+		UE_LOG(LogText3D, Error, TEXT("Text3D renderer is not valid, cannot proceed"));
+		return;
 	}
 
-	const float TotalHeight = GetTextHeight();
-	if (VerticalAlignment != EText3DVerticalTextAlignment::FirstLine)
+	if (!Font)
 	{
-		// First align it to Top
-		VerticalOffset -= ShapedText->FontAscender;
-
-		if (VerticalAlignment == EText3DVerticalTextAlignment::Center)
+		if (const UText3DProjectSettings* Text3DSettings = UText3DProjectSettings::Get())
 		{
-			VerticalOffset += TotalHeight * 0.5f;
+			Font = Text3DSettings->GetFallbackFont();
 		}
-		else if (VerticalAlignment == EText3DVerticalTextAlignment::Bottom)
+
+		if (!Font)
 		{
-			VerticalOffset += TotalHeight + ShapedText->FontDescender;
+			UE_LOG(LogText3D, Error, TEXT("Font object is not valid, Fallback font is not defined in project settings, cannot proceed"));
+			return;
 		}
 	}
+	
+	if (bIsUpdatingText)
+	{
+		return;
+	}
 
-	VerticalOffset -= LineIndex * (ShapedText->LineHeight + LineSpacing);
+	bIsUpdatingText = true;
+	ON_SCOPE_EXIT { bIsUpdatingText = false; };
 
-	return FVector(0.0f, HorizontalOffset, VerticalOffset);
-}
+	TextPreUpdateDelegate.Broadcast(this, UpdateFlags);
 
-void UText3DComponent::UpdateTransforms()
-{
-	ShapedText->Kerning = Kerning;
-	ShapedText->WordSpacing = WordSpacing;
-	ShapedText->MaxWidth = MaxWidth;
-	ShapedText->bWrap = bHasMaxWidth && MaxWidthHandling == EText3DMaxWidthHandling::WrapAndScale;
+	TRACE_CPUPROFILER_EVENT_SCOPE(UText3DComponent::RebuildInternal);
 
-	ShapedText->CalculateWidth();
-	CalculateTextScale();
+	bool bBroadcastEvent = false;
 
-	if (TextRenderer)
+	if (EnumHasAnyFlags(UpdateFlags,EText3DRendererFlags::Geometry))
 	{
-		TextRenderer->Update(EText3DRendererFlags::Layout);
-	}
+		if (!FormattedText.IsSet())
+		{
+			FormattedText = Text;
+			FormatText(FormattedText.GetValue());
+		}
 
-	ModifyFlags &= ~EText3DModifyFlags::Layout;
-}
+		TextRenderer->Clear();
 
-void UText3DComponent::ClearTextMesh()
-{
-	if (TextRenderer)
+		bBroadcastEvent = true;
+	}
+
+	TextRenderer->Update(UpdateFlags);
+	
+	if (bBroadcastEvent)
 	{
-		TextRenderer->Clear();
+		UpdateStatistics();
+
+		TextGeneratedNativeDelegate.Broadcast();
+		TextGeneratedDelegate.Broadcast();
 	}
-}
 
-void UText3DComponent::TriggerInternalRebuild(const EText3DModifyFlags InModifyFlags)
-{
-	if (EnumHasAnyFlags(InModifyFlags, EText3DModifyFlags::Geometry))
+	if (EnumHasAnyFlags(UpdateFlags,EText3DRendererFlags::Visibility))
 	{
-		MarkForGeometryUpdate();
+		MarkRenderStateDirty();
 	}
 
-	if (EnumHasAnyFlags(InModifyFlags, EText3DModifyFlags::Layout))
+	if (bCleanCache)
 	{
-		MarkForLayoutUpdate();
+		if (UText3DEngineSubsystem* Subsystem = UText3DEngineSubsystem::Get())
+		{
+			Subsystem->Cleanup();
+		}
 	}
 
-	RebuildInternal();
+	TextPostUpdateDelegate.Broadcast(this, UpdateFlags);
+
+	ClearUpdateFlags();
 }
 
 void UText3DComponent::OnTextRendererClassChanged()
@@ -934,7 +906,7 @@ void UText3DComponent::OnTextRendererClassChanged()
 	if (IsRegistered())
 	{
 		TextRenderer->Create();
-		RebuildInternal(/** AutoUpdate */true, /** CleanCache */true);
+		RequestUpdate(EText3DRendererFlags::All);
 	}
 }
 
@@ -943,247 +915,40 @@ void UText3DComponent::OnTextChanged()
 	FormattedText = Text;
 	FormatText(FormattedText.GetValue());
 	
-	MarkForGeometryUpdate();
-	RebuildInternal();
+	RequestUpdate(EText3DRendererFlags::All);
 }
 
 void UText3DComponent::OnFontPropertiesChanged()
 {
 	RefreshTypeface();
 
-	MarkForGeometryUpdate();
-	RebuildInternal();
-}
-
-void UText3DComponent::OnGeometryPropertiesChanged()
-{
-	// Extrude
-	
-	Extrude = FMath::Max(0.0f, Extrude);
-
-	// Bevels
-	
-	Bevel = FMath::Clamp(Bevel, 0.f, GetMaxBevel());
-
-	int32 MaxBevelSegments = 1;
-	int32 MinBevelSegments = 1;
-	switch (BevelType)
-	{
-	case EText3DBevelType::Linear:
-	case EText3DBevelType::OneStep:
-	case EText3DBevelType::TwoSteps:
-	case EText3DBevelType::Engraved:
-		{
-			MaxBevelSegments = 1;
-			break;
-		}
-	case EText3DBevelType::Convex:
-	case EText3DBevelType::Concave:
-		{
-			MaxBevelSegments = 8;
-			break;
-		}
-	case EText3DBevelType::HalfCircle:
-		{
-			MaxBevelSegments = 16;
-			MinBevelSegments = 2;
-			break;
-		}
-	}
-
-	BevelSegments = FMath::Clamp(BevelSegments, MinBevelSegments, MaxBevelSegments);
-
-	MarkForGeometryUpdate();
-	RebuildInternal();
-}
-
-void UText3DComponent::OnLayoutPropertiesChanged()
-{
-	MaxWidth = FMath::Max(1.0f, MaxWidth);
-	MaxHeight = FMath::Max(1.0f, MaxHeight);
-
-	UpdateTransforms();
-}
-
-void UText3DComponent::OnLayoutHandlingPropertiesChanged()
-{
-	if (MaxWidthHandling == EText3DMaxWidthHandling::WrapAndScale)
-	{
-		OnGeometryPropertiesChanged();
-	}
-	else
-	{
-		OnLayoutPropertiesChanged();
-	}
-}
-
-void UText3DComponent::OnVisibilityPropertiesChanged()
-{
-	if (TextRenderer)
-	{
-		TextRenderer->Update(EText3DRendererFlags::Visibility);
-	}
-
-	MarkRenderStateDirty();
-}
-
-void UText3DComponent::BuildTextMesh(const bool& bCleanCache)
-{
-	// If we're already building, or have a build pending, don't do anything.
-	if (bIsBuilding)
-	{
-		return;
-	}
-
-	bIsBuilding = true;
-
-	TWeakObjectPtr<UText3DComponent> WeakThis(this);
-
-	// Execution guarded by the above flag
-	AsyncTask(ENamedThreads::GameThread, [WeakThis, bCleanCache]()
-	{
-		if (UText3DComponent* StrongThis = WeakThis.Get())
-		{
-			if (!UE::IsSavingPackage(StrongThis))
-			{
-				StrongThis->BuildTextMeshInternal(bCleanCache);
-			}
-		}
-	});
+	RequestUpdate(EText3DRendererFlags::All);
 }
 
-void UText3DComponent::BuildTextMeshInternal(const bool& bCleanCache)
+void UText3DComponent::ScheduleTextUpdateNextTick()
 {
-	TRACE_CPUPROFILER_EVENT_SCOPE(UText3DComponent::Rebuild);
-
-	ON_SCOPE_EXIT { bIsBuilding = false; };
-
-	if (!IsRegistered())
-	{
-		return;
-	}
-
-	if (!TextRenderer)
-	{
-		UE_LOG(LogText3D, Error, TEXT("Text renderer is not valid"));
-		return;
-	}
-
-	ClearTextMesh();
-
-	if (!Font)
-	{
-		UE_LOG(LogText3D, Error, TEXT("Font object is not valid"));
-		return;
-	}
-
-	UText3DEngineSubsystem* Subsystem = UText3DEngineSubsystem::Get();
-
-	const uint32 TypefaceIndex = GetTypeFaceIndex();
-	FCachedFontData& CachedFontData = Subsystem->GetCachedFontData(Font, TypefaceIndex);
-	const FT_Face Face = CachedFontData.GetFreeTypeFace(TypefaceIndex);
-	if (!Face)
+	if (!TextUpdateHandle.IsValid())
 	{
-		UE_LOG(LogText3D, Error, TEXT("Failed to load font data '%s'"), *CachedFontData.GetFontName());
-		return;
-	}
-
-	ShapedText->Reset();
-	ShapedText->LineHeight = Face->size->metrics.height * FontInverseScale;
-	ShapedText->FontAscender = Face->size->metrics.ascender * FontInverseScale;
-	ShapedText->FontDescender = Face->size->metrics.descender * FontInverseScale;
-	ShapedText->Kerning = Kerning;
-	ShapedText->WordSpacing = WordSpacing;
-	ShapedText->MaxWidth = MaxWidth;
-	ShapedText->bWrap = bHasMaxWidth && MaxWidthHandling == EText3DMaxWidthHandling::WrapAndScale;
-
-	constexpr int32 AdjustedFontSize = 48; // Magic number that makes font scale consistent with previous implementation
-	FSlateFontInfo FontInfo(Font, AdjustedFontSize);
-	FontInfo.CompositeFont = FStyleDefaults::GetFontInfo().CompositeFont;
-	FontInfo.TypefaceFontName = Typeface;
-
-	FTextBlockStyle Style;
-	Style.SetFont(FontInfo);
-
-	if (!TextLayout.IsValid())
-	{
-		TextLayout = MakeShared<FText3DLayout>(Style);
-	}
-
-	if (!TextLayoutMarshaller.IsValid())
-	{
-		TextLayoutMarshaller = FPlainTextLayoutMarshaller::Create();
-	}
-
-	if (!FormattedText.IsSet())
-	{
-		FormattedText = Text;
-		FormatText(FormattedText.GetValue());
-	}
-
-	FText3DLayoutShaper::Get()->ShapeBidirectionalText(Style, FormattedText.GetValue().ToString(), TextLayout, TextLayoutMarshaller, ShapedText->Lines);
-
-	const int32 ApproximateGlyphNum = Algo::TransformAccumulate(ShapedText->Lines,
-		[](const FShapedGlyphLine& InLine)
-		{
-			return InLine.GlyphsToRender.Num();
-		},
-		0);
-
-	ShapedText->GlyphIndexToFontFace.Reserve(ApproximateGlyphNum);
-
-	for (int32 LineIndex = 0; LineIndex < ShapedText->Lines.Num(); ++LineIndex)
-	{
-		for (const FShapedGlyphEntry& GlyphEntry : ShapedText->Lines[LineIndex].GlyphsToRender)
-		{
-			if (!GlyphEntry.FontFaceData.IsValid())
-			{
-				// Add as nullptr if not already in map
-				if (!ShapedText->GlyphIndexToFontFace.Contains(GlyphEntry.GlyphIndex))
+		TextUpdateHandle = FTSTicker::GetCoreTicker().AddTicker(
+			FTickerDelegate::CreateWeakLambda(
+				this,
+				[this](float)->bool
 				{
-					ShapedText->GlyphIndexToFontFace.Add(GlyphEntry.GlyphIndex, nullptr);
-				}
+					if (UE::IsSavingPackage(this))
+					{
+						// Postpone to next tick
+						return true;
+					}
 
-				continue;
-			}
-
-			const TSharedPtr<FFreeTypeFace> FontFacePtr = GlyphEntry.FontFaceData->FontFace.Pin();
-
-			if (FontFacePtr.IsValid())
-			{
-				ShapedText->GlyphIndexToFontFace.FindOrAdd(GlyphEntry.GlyphIndex, FontFacePtr.Get());
-			}
-		}
-	}
+					RebuildInternal();
 
-	ShapedText->CalculateWidth();
-	CalculateTextScale();
+					TextUpdateHandle.Reset();
 
-	TextRenderer->Update(EText3DRendererFlags::All);
-
-	UpdateStatistics();
-
-	TextGeneratedNativeDelegate.Broadcast();
-	TextGeneratedDelegate.Broadcast();
-
-	ClearUpdateFlags();
-
-	if (bCleanCache)
-	{
-		Subsystem->Cleanup();
-	}
-}
-
-float UText3DComponent::GetMaxBevel() const
-{
-	return Extrude / 2.0f;
-}
-
-void UText3DComponent::OnMaterialChanged()
-{
-	if (TextRenderer)
-	{
-		TextRenderer->Update(EText3DRendererFlags::Material);
+					// Executes only once
+					return false;
+				}
+			)
+		);
 	}
 }
 
@@ -1206,23 +971,25 @@ void UText3DComponent::OnVisibilityChanged()
 {
 	Super::OnVisibilityChanged();
 
-	if (TextRenderer)
-	{
-		TextRenderer->Update(EText3DRendererFlags::Visibility);
-	}
+	RequestUpdate(EText3DRendererFlags::Visibility);
 }
 
 void UText3DComponent::OnHiddenInGameChanged()
 {
 	Super::OnHiddenInGameChanged();
 
-	if (TextRenderer)
+	RequestUpdate(EText3DRendererFlags::Visibility);
+}
+
+void UText3DComponent::FormatText(FText& InOutText) const
+{
+	if (bEnforceUpperCase)
 	{
-		TextRenderer->Update(EText3DRendererFlags::Visibility);
+		InOutText = InOutText.ToUpper();
 	}
 }
 
-void UText3DComponent::GetBounds(FVector& Origin, FVector& BoxExtent)
+void UText3DComponent::GetBounds(FVector& Origin, FVector& BoxExtent) const
 {
 	FBox Box(ForceInit);
 
@@ -1234,4 +1001,14 @@ void UText3DComponent::GetBounds(FVector& Origin, FVector& BoxExtent)
 	Box.GetCenterAndExtents(Origin, BoxExtent);
 }
 
-#undef LOCTEXT_NAMESPACE
\ No newline at end of file
+FBox UText3DComponent::GetBounds() const
+{
+	if (TextRenderer)
+	{
+		return TextRenderer->GetBounds();
+	}
+
+	return FBox(ForceInit);
+}
+
+#undef LOCTEXT_NAMESPACE
diff --git a/Engine/Plugins/Experimental/Text3D/Source/Text3D/Public/Characters/Text3DCharacterTransform.h b/Engine/Plugins/Experimental/Text3D/Source/Text3D/Public/Characters/Text3DCharacterTransform.h
--- a/Engine/Plugins/Experimental/Text3D/Source/Text3D/Public/Characters/Text3DCharacterTransform.h
+++ b/Engine/Plugins/Experimental/Text3D/Source/Text3D/Public/Characters/Text3DCharacterTransform.h
@@ -2,25 +2,185 @@
 
 #pragma once
 
-#include "Components/ActorComponent.h"
 #include "Components/SceneComponent.h"
 #include "Text3DTypes.h"
 #include "Text3DCharacterTransform.generated.h"
 
-UCLASS(ClassGroup=(Text3D), HideCategories = (Collision, Tags, Activation, Cooking), Meta = (BlueprintSpawnableComponent))
-class TEXT3D_API UText3DCharacterTransform : public USceneComponent
+class UText3DComponent;
+
+UCLASS(MinimalAPI, ClassGroup=(Text3D), HideCategories = (Collision, Tags, Activation, Cooking, Rendering, Physics, Mobility, LOD, AssetUserData, Navigation, Transform), Meta = (BlueprintSpawnableComponent))
+class UText3DCharacterTransform : public USceneComponent
 {
 	GENERATED_BODY()
 
 public:	
 	UText3DCharacterTransform();
 
+	bool GetLocationEnabled() const
+	{
+		return bLocationEnabled;
+	}
+
+	float GetLocationProgress() const
+	{
+		return LocationProgress;
+	}
+
+	EText3DCharacterEffectOrder GetLocationOrder() const
+	{
+		return LocationOrder;
+	}
+
+	float GetLocationRange() const
+	{
+		return LocationRange;
+	}
+
+	FVector GetLocationDistance() const
+	{
+		return LocationDistance;
+	}
+
+	bool GetRotationEnabled() const
+	{
+		return bRotateEnabled;
+	}
+
+	float GetRotationProgress() const
+	{
+		return RotateProgress;
+	}
+
+	EText3DCharacterEffectOrder GetRotationOrder() const
+	{
+		return RotateOrder;
+	}
+
+	float GetRotationRange() const
+	{
+		return RotateRange;
+	}
+
+	FRotator GetRotationBegin() const
+	{
+		return RotateBegin;
+	}
+
+	FRotator GetRotationEnd() const
+	{
+		return RotateEnd;
+	}
+
+	bool GetScaleEnabled() const
+	{
+		return bScaleEnabled;
+	}
+
+	float GetScaleProgress() const
+	{
+		return ScaleProgress;
+	}
+
+	EText3DCharacterEffectOrder GetScaleOrder() const
+	{
+		return ScaleOrder;
+	}
+
+	float GetScaleRange() const
+	{
+		return ScaleRange;
+	}
+
+	FVector GetScaleBegin() const
+	{
+		return ScaleBegin;
+	}
+
+	FVector GetScaleEnd() const
+	{
+		return ScaleEnd;
+	}
+
+	// Location
+
+	UFUNCTION(BlueprintCallable, Category = "Location")
+	TEXT3D_API void SetLocationEnabled(bool bEnabled);
+
+	UFUNCTION(BlueprintCallable, Category = "Location")
+	TEXT3D_API void SetLocationProgress(float Progress);
+
+	UFUNCTION(BlueprintCallable, Category = "Location")
+	TEXT3D_API void SetLocationOrder(EText3DCharacterEffectOrder Order);
+
+	UFUNCTION(BlueprintCallable, Category = "Location")
+	TEXT3D_API void SetLocationRange(float Range);
+
+	UFUNCTION(BlueprintCallable, Category = "Location")
+	TEXT3D_API void SetLocationDistance(FVector Distance);
+
+	// Rotation
+
+	UFUNCTION(BlueprintCallable, Category = "Rotate")
+	TEXT3D_API void SetRotateEnabled(bool bEnabled);
+
+	UFUNCTION(BlueprintCallable, Category = "Rotate")
+	TEXT3D_API void SetRotateProgress(float Progress);
+
+	UFUNCTION(BlueprintCallable, Category = "Rotate")
+	TEXT3D_API void SetRotateOrder(EText3DCharacterEffectOrder Order);
+
+	UFUNCTION(BlueprintCallable, Category = "Rotate")
+	TEXT3D_API void SetRotateRange(float Range);
+
+	UFUNCTION(BlueprintCallable, Category = "Rotate")
+	TEXT3D_API void SetRotateBegin(FRotator Value);
+
+	UFUNCTION(BlueprintCallable, Category = "Rotate")
+	TEXT3D_API void SetRotateEnd(FRotator Value);
+
+	// Scale
+
+	UFUNCTION(BlueprintCallable, Category = "Scale")
+	TEXT3D_API void SetScaleEnabled(bool bEnabled);
+
+	UFUNCTION(BlueprintCallable, Category = "Scale")
+	TEXT3D_API void SetScaleProgress(float Progress);
+
+	UFUNCTION(BlueprintCallable, Category = "Scale")
+	TEXT3D_API void SetScaleOrder(EText3DCharacterEffectOrder Order);
+
+	UFUNCTION(BlueprintCallable, Category = "Scale")
+	TEXT3D_API void SetScaleRange(float Range);
+
+	UFUNCTION(BlueprintCallable, Category = "Scale")
+	TEXT3D_API void SetScaleBegin(FVector Value);
+
+	UFUNCTION(BlueprintCallable, Category = "Scale")
+	TEXT3D_API void SetScaleEnd(FVector Value);
+
+protected:
+	//~ Begin USceneComponent
 	virtual void OnRegister() override;
 	virtual void OnUnregister() override;
+	//~ End USceneComponent
 
+	//~ Begin UObject
 #if WITH_EDITOR
 	virtual void PostEditChangeProperty(struct FPropertyChangedEvent& PropertyChangedEvent) override;
 #endif
+	//~ End UObject
+
+	UText3DComponent* GetText3DComponent();
+
+	void ProcessEffect();
+
+	float GetEffectValue(int32 Index, int32 Total, EText3DCharacterEffectOrder Order, float Progress, float Range);
+	void GetLineParameters(float Range, EText3DCharacterEffectOrder Order, int32 Count, float & EffectOut, float & StripOut);
+	int32 GetEffectPosition(int32 Index, int32 Total, EText3DCharacterEffectOrder Order);
+
+	void ResetLocation();
+	void ResetRotate();
+	void ResetScale();
 
 	// Location
 
@@ -39,61 +199,5 @@
 	UPROPERTY(EditAnywhere, BlueprintReadWrite, BlueprintSetter = SetLocationDistance, Category = "Location", meta = (EditCondition = "bLocationEnabled"))
 	FVector LocationDistance;
 
-	UFUNCTION(BlueprintCallable, Category = "Location")
-	void SetLocationEnabled(bool bEnabled);
-
-	UFUNCTION(BlueprintCallable, Category = "Location")
-	void SetLocationProgress(float Progress);
-
-	UFUNCTION(BlueprintCallable, Category = "Location")
-	void SetLocationOrder(EText3DCharacterEffectOrder Order);
-
-	UFUNCTION(BlueprintCallable, Category = "Location")
-	void SetLocationRange(float Range);
-
-	UFUNCTION(BlueprintCallable, Category = "Location")
-	void SetLocationDistance(FVector Distance);
-
-
-	// Scale
-
-	UPROPERTY(EditAnywhere, BlueprintReadWrite, BlueprintSetter = SetScaleEnabled, Category = "Scale")
-	bool bScaleEnabled;
-
-	UPROPERTY(EditAnywhere, BlueprintReadWrite, BlueprintSetter = SetScaleProgress, Category = "Scale", meta = (EditCondition = "bScaleEnabled", ClampMin = 0, ClampMax = 100))
-	float ScaleProgress;
-
-	UPROPERTY(EditAnywhere, BlueprintReadWrite, BlueprintSetter = SetScaleOrder, Category = "Scale", meta = (EditCondition = "bScaleEnabled"))
-	EText3DCharacterEffectOrder ScaleOrder;
-
-	UPROPERTY(EditAnywhere, BlueprintReadWrite, BlueprintSetter = SetScaleRange, Category = "Scale", meta = (EditCondition = "bScaleEnabled", ClampMin = 0, ClampMax = 100))
-	float ScaleRange;
-
-	UPROPERTY(EditAnywhere, BlueprintReadWrite, BlueprintSetter = SetScaleBegin, Category = "Scale", meta = (EditCondition = "bScaleEnabled", ClampMin = 0))
-	FVector ScaleBegin;
-
-	UPROPERTY(EditAnywhere, BlueprintReadWrite, BlueprintSetter = SetScaleEnd, Category = "Scale", meta = (EditCondition = "bScaleEnabled", ClampMin = 0))
-	FVector ScaleEnd;
-
-
-	UFUNCTION(BlueprintCallable, Category = "Scale")
-	void SetScaleEnabled(bool bEnabled);
-
-	UFUNCTION(BlueprintCallable, Category = "Scale")
-	void SetScaleProgress(float Progress);
-
-	UFUNCTION(BlueprintCallable, Category = "Scale")
-	void SetScaleOrder(EText3DCharacterEffectOrder Order);
-
-	UFUNCTION(BlueprintCallable, Category = "Scale")
-	void SetScaleRange(float Range);
-
-	UFUNCTION(BlueprintCallable, Category = "Scale")
-	void SetScaleBegin(FVector Value);
-
-	UFUNCTION(BlueprintCallable, Category = "Scale")
-	void SetScaleEnd(FVector Value);
-
-
 	// Rotate
 
@@ -115,37 +219,25 @@
 	UPROPERTY(EditAnywhere, BlueprintReadWrite, BlueprintSetter = SetRotateEnd, Category = "Rotate", meta = (EditCondition = "bRotateEnabled"))
 	FRotator RotateEnd;
 
+	// Scale
 
-	UFUNCTION(BlueprintCallable, Category = "Rotate")
-	void SetRotateEnabled(bool bEnabled);
+	UPROPERTY(EditAnywhere, BlueprintReadWrite, BlueprintSetter = SetScaleEnabled, Category = "Scale")
+	bool bScaleEnabled;
 
-	UFUNCTION(BlueprintCallable, Category = "Rotate")
-	void SetRotateProgress(float Progress);
+	UPROPERTY(EditAnywhere, BlueprintReadWrite, BlueprintSetter = SetScaleProgress, Category = "Scale", meta = (EditCondition = "bScaleEnabled", ClampMin = 0, ClampMax = 100))
+	float ScaleProgress;
 
-	UFUNCTION(BlueprintCallable, Category = "Rotate")
-	void SetRotateOrder(EText3DCharacterEffectOrder Order);
+	UPROPERTY(EditAnywhere, BlueprintReadWrite, BlueprintSetter = SetScaleOrder, Category = "Scale", meta = (EditCondition = "bScaleEnabled"))
+	EText3DCharacterEffectOrder ScaleOrder;
 
-	UFUNCTION(BlueprintCallable, Category = "Rotate")
-	void SetRotateRange(float Range);
+	UPROPERTY(EditAnywhere, BlueprintReadWrite, BlueprintSetter = SetScaleRange, Category = "Scale", meta = (EditCondition = "bScaleEnabled", ClampMin = 0, ClampMax = 100))
+	float ScaleRange;
 
-	UFUNCTION(BlueprintCallable, Category = "Rotate")
-	void SetRotateBegin(FRotator Value);
+	UPROPERTY(EditAnywhere, BlueprintReadWrite, BlueprintSetter = SetScaleBegin, Category = "Scale", meta = (EditCondition = "bScaleEnabled", ClampMin = 0))
+	FVector ScaleBegin;
 
-	UFUNCTION(BlueprintCallable, Category = "Rotate")
-	void SetRotateEnd(FRotator Value);
+	UPROPERTY(EditAnywhere, BlueprintReadWrite, BlueprintSetter = SetScaleEnd, Category = "Scale", meta = (EditCondition = "bScaleEnabled", ClampMin = 0))
+	FVector ScaleEnd;
 
-private:
-	class UText3DComponent* GetText3DComponent();
-
-	void ProcessEffect();
-
-	float GetEffectValue(int32 Index, int32 Total, EText3DCharacterEffectOrder Order, float Progress, float Range);
-	void GetLineParameters(float Range, EText3DCharacterEffectOrder Order, int32 Count, float & EffectOut, float & StripOut);
-	int32 GetEffectPosition(int32 Index, int32 Total, EText3DCharacterEffectOrder Order);
-
-	void ResetLocation();
-	void ResetRotate();
-	void ResetScale();
-
 	bool bInitialized;
 };
diff --git a/Engine/Plugins/Experimental/Text3D/Source/Text3D/Public/Extensions/Text3DDefaultTransformExtension.h b/Engine/Plugins/Experimental/Text3D/Source/Text3D/Public/Extensions/Text3DDefaultTransformExtension.h
new file mode 100644
--- /dev/null
+++ b/Engine/Plugins/Experimental/Text3D/Source/Text3D/Public/Extensions/Text3DDefaultTransformExtension.h
@@ -0,0 +1,226 @@
+// Copyright Epic Games, Inc. All Rights Reserved.
+
+#pragma once
+
+#include "Text3DTransformExtensionBase.h"
+#include "Text3DTypes.h"
+#include "Text3DDefaultTransformExtension.generated.h"
+
+/** Extension that handles transform data for Text3D */
+UCLASS(MinimalAPI, EditInlineNew)
+class UText3DDefaultTransformExtension : public UText3DTransformExtensionBase
+{
+	GENERATED_BODY()
+
+public:
+	UFUNCTION(BlueprintCallable, Category = "Text3D|Transform")
+	TEXT3D_API void SetLocationEnabled(bool bEnabled);
+
+	bool GetLocationEnabled() const
+	{
+		return bLocationEnabled;
+	}
+
+	UFUNCTION(BlueprintCallable, Category = "Text3D|Transform")
+	TEXT3D_API void SetLocationProgress(float Progress);
+
+	float GetLocationProgress() const
+	{
+		return LocationProgress;
+	}
+
+	UFUNCTION(BlueprintCallable, Category = "Text3D|Transform")
+	TEXT3D_API void SetLocationOrder(EText3DCharacterEffectOrder Order);
+
+	EText3DCharacterEffectOrder GetLocationOrder() const
+	{
+		return LocationOrder;
+	}
+
+	UFUNCTION(BlueprintCallable, Category = "Text3D|Transform")
+	TEXT3D_API void SetLocationRange(float Range);
+
+	float GetLocationRange() const
+	{
+		return LocationRange;
+	}
+
+	UFUNCTION(BlueprintCallable, Category = "Text3D|Transform")
+	TEXT3D_API void SetLocationDistance(const FVector& Distance);
+
+	FVector GetLocationDistance() const
+	{
+		return LocationDistance;
+	}
+
+	UFUNCTION(BlueprintCallable, Category = "Text3D|Transform")
+	TEXT3D_API void SetRotationEnabled(bool bEnabled);
+
+	bool GetRotationEnabled() const
+	{
+		return bRotationEnabled;
+	}
+
+	UFUNCTION(BlueprintCallable, Category = "Text3D|Transform")
+	TEXT3D_API void SetRotationProgress(float Progress);
+
+	float GetRotationProgress() const
+	{
+		return RotationProgress;
+	}
+
+	UFUNCTION(BlueprintCallable, Category = "Text3D|Transform")
+	TEXT3D_API void SetRotationOrder(EText3DCharacterEffectOrder Order);
+
+	EText3DCharacterEffectOrder GetRotationOrder() const
+	{
+		return RotationOrder;
+	}
+
+	UFUNCTION(BlueprintCallable, Category = "Text3D|Transform")
+	TEXT3D_API void SetRotationRange(float Range);
+
+	float GetRotationRange() const
+	{
+		return RotationRange;
+	}
+
+	UFUNCTION(BlueprintCallable, Category = "Text3D|Transform")
+	TEXT3D_API void SetRotationBegin(const FRotator& Value);
+
+	FRotator GetRotationBegin() const
+	{
+		return RotationBegin;
+	}
+
+	UFUNCTION(BlueprintCallable, Category = "Text3D|Transform")
+	TEXT3D_API void SetRotationEnd(const FRotator& Value);
+
+	FRotator GetRotationEnd() const
+	{
+		return RotationEnd;
+	}
+
+	UFUNCTION(BlueprintCallable, Category = "Text3D|Transform")
+	TEXT3D_API void SetScaleEnabled(bool bEnabled);
+
+	bool GetScaleEnabled() const
+	{
+		return bScaleEnabled;
+	}
+
+	UFUNCTION(BlueprintCallable, Category = "Text3D|Transform")
+	TEXT3D_API void SetScaleProgress(float Progress);
+
+	float GetScaleProgress() const
+	{
+		return ScaleProgress;
+	}
+
+	UFUNCTION(BlueprintCallable, Category = "Text3D|Transform")
+	TEXT3D_API void SetScaleOrder(EText3DCharacterEffectOrder Order);
+
+	EText3DCharacterEffectOrder GetScaleOrder() const
+	{
+		return ScaleOrder;
+	}
+
+	UFUNCTION(BlueprintCallable, Category = "Text3D|Transform")
+	TEXT3D_API void SetScaleRange(float Range);
+
+	float GetScaleRange() const
+	{
+		return ScaleRange;
+	}
+
+	UFUNCTION(BlueprintCallable, Category = "Text3D|Transform")
+	TEXT3D_API void SetScaleBegin(const FVector& Value);
+
+	FVector GetScaleBegin() const
+	{
+		return ScaleBegin;
+	}
+
+	UFUNCTION(BlueprintCallable, Category = "Text3D|Transform")
+	TEXT3D_API void SetScaleEnd(const FVector& Value);
+
+	FVector GetScaleEnd() const
+	{
+		return ScaleEnd;
+	}
+
+protected:
+	UText3DDefaultTransformExtension();
+
+	//~ Begin UObject
+#if WITH_EDITOR
+	virtual void PostEditChangeProperty(FPropertyChangedEvent& InEvent) override;
+#endif
+	//~ End UObject
+
+	//~ Begin UText3DTransformExtensionBase
+	virtual void ApplyTransform(uint32 InGlyphIndex, uint32 InGlyphCount, FTransform& InGlyphTransform) override;
+	//~ End UText3DTransformExtensionBase
+
+	void OnTransformOptionsChanged();
+	float GetEffectValue(int32 Index, int32 Total, EText3DCharacterEffectOrder Order, float Progress, float Range);
+	void GetLineParameters(float Range, EText3DCharacterEffectOrder Order, int32 Total, float & EffectOut, float & StripOut);
+	int32 GetEffectPosition(int32 Index, int32 Total, EText3DCharacterEffectOrder Order);
+
+	// Location
+
+	UPROPERTY(EditAnywhere, Getter="GetLocationEnabled", Setter="SetLocationEnabled", Category = "Transform")
+	bool bLocationEnabled = false;
+
+	UPROPERTY(EditAnywhere, Getter, Setter, Category = "Transform", meta = (EditCondition = "bLocationEnabled", ClampMin = 0, ClampMax = 100))
+	float LocationProgress = 0.f;
+
+	UPROPERTY(EditAnywhere, Getter, Setter, Category = "Transform", meta = (EditCondition = "bLocationEnabled"))
+	EText3DCharacterEffectOrder LocationOrder = EText3DCharacterEffectOrder::Normal;
+
+	UPROPERTY(EditAnywhere, Getter, Setter, Category = "Transform", meta = (EditCondition = "bLocationEnabled", ClampMin = 0, ClampMax = 100))
+	float LocationRange = 50.f;
+
+	UPROPERTY(EditAnywhere, Getter, Setter, Category = "Transform", meta = (EditCondition = "bLocationEnabled"))
+	FVector LocationDistance = FVector(100.f, 0.f, 0.f);
+
+	// Rotate
+
+	UPROPERTY(EditAnywhere, Getter="GetRotationEnabled", Setter="SetRotationEnabled", Category = "Transform")
+	bool bRotationEnabled = false;
+
+	UPROPERTY(EditAnywhere, Getter, Setter, Category = "Transform", meta = (EditCondition = "bRotationEnabled", ClampMin = 0, ClampMax = 100))
+	float RotationProgress = 0.0f;
+
+	UPROPERTY(EditAnywhere, Getter, Setter, Category = "Transform", meta = (EditCondition = "bRotationEnabled"))
+	EText3DCharacterEffectOrder RotationOrder = EText3DCharacterEffectOrder::Normal;
+
+	UPROPERTY(EditAnywhere, Getter, Setter, Category = "Transform", meta = (EditCondition = "bRotationEnabled", ClampMin = 0, ClampMax = 100))
+	float RotationRange = 50.f;
+
+	UPROPERTY(EditAnywhere, Getter, Setter, Category = "Transform", meta = (EditCondition = "bRotationEnabled"))
+	FRotator RotationBegin = FRotator(-90.f, 0.f, 0.f);
+
+	UPROPERTY(EditAnywhere, Getter, Setter, Category = "Transform", meta = (EditCondition = "bRotationEnabled"))
+	FRotator RotationEnd = FRotator(0.f, 0.f, 0.f);
+
+	// Scale
+
+	UPROPERTY(EditAnywhere, Getter="GetScaleEnabled", Setter="SetScaleEnabled", Category = "Transform")
+	bool bScaleEnabled = false;
+
+	UPROPERTY(EditAnywhere, Getter, Setter, Category = "Transform", meta = (EditCondition = "bScaleEnabled", ClampMin = 0, ClampMax = 100))
+	float ScaleProgress = 0.f;
+
+	UPROPERTY(EditAnywhere, Getter, Setter, Category = "Transform", meta = (EditCondition = "bScaleEnabled"))
+	EText3DCharacterEffectOrder ScaleOrder = EText3DCharacterEffectOrder::Normal;
+
+	UPROPERTY(EditAnywhere, Getter, Setter, Category = "Transform", meta = (EditCondition = "bScaleEnabled", ClampMin = 0, ClampMax = 100))
+	float ScaleRange = 50.f;
+
+	UPROPERTY(EditAnywhere, Getter, Setter, Category = "Transform", meta = (EditCondition = "bScaleEnabled", ClampMin = 0))
+	FVector ScaleBegin = FVector(1.f, 0.f, 0.f);
+
+	UPROPERTY(EditAnywhere, Getter, Setter, Category = "Transform", meta = (EditCondition = "bScaleEnabled", ClampMin = 0))
+	FVector ScaleEnd = FVector::OneVector;
+};
diff --git a/Engine/Plugins/Experimental/Text3D/Source/Text3D/Public/Extensions/Text3DTransformExtensionBase.h b/Engine/Plugins/Experimental/Text3D/Source/Text3D/Public/Extensions/Text3DTransformExtensionBase.h
new file mode 100644
--- /dev/null
+++ b/Engine/Plugins/Experimental/Text3D/Source/Text3D/Public/Extensions/Text3DTransformExtensionBase.h
@@ -0,0 +1,62 @@
+// Copyright Epic Games, Inc. All Rights Reserved.
+
+#pragma once
+
+#include "Text3DExtensionBase.h"
+#include "Text3DTypes.h"
+#include "Text3DTransformExtensionBase.generated.h"
+
+struct FText3DRange
+{
+	static constexpr uint16 Unbound = std::numeric_limits<uint16>::max();
+
+	FText3DRange() = default;
+	explicit FText3DRange(uint16 InStartIndex, uint16 InCount)
+		: StartIndex(InStartIndex)
+		, Count(InCount)
+	{
+	}
+
+	bool IsUnbound() const
+	{
+		return IsStartUnbound() && IsEndUnbound();
+	}
+
+	bool IsEndUnbound() const
+	{
+		return Count == Unbound;
+	}
+
+	bool IsStartUnbound() const
+	{
+		return StartIndex == Unbound;
+	}
+
+	bool IsInRange(uint16 InIndex) const
+	{
+		return IsUnbound()
+			|| (IsStartUnbound() && InIndex < Count)
+			|| (IsEndUnbound() && InIndex >= StartIndex)
+			|| (InIndex >= StartIndex && InIndex <= (StartIndex + Count));
+	}
+
+	uint16 StartIndex = Unbound;
+	uint16 Count = Unbound;
+};
+
+/** Extension that handles transform data for Text3D */
+UCLASS(MinimalAPI, Abstract)
+class UText3DTransformExtensionBase : public UText3DExtensionBase
+{
+	GENERATED_BODY()
+
+public:
+	virtual FText3DRange GetTargetRange() const
+	{
+		return FText3DRange();
+	}
+
+	virtual void ApplyTransform(uint32 InGlyphIndex, uint32 InGlyphCount, FTransform& InGlyphTransform)
+	{
+	}
+};
diff --git a/Engine/Plugins/Experimental/Text3D/Source/Text3D/Public/Renderers/Text3DRendererBase.h b/Engine/Plugins/Experimental/Text3D/Source/Text3D/Public/Renderers/Text3DRendererBase.h
--- a/Engine/Plugins/Experimental/Text3D/Source/Text3D/Public/Renderers/Text3DRendererBase.h
+++ b/Engine/Plugins/Experimental/Text3D/Source/Text3D/Public/Renderers/Text3DRendererBase.h
@@ -3,24 +3,11 @@
 #pragma once
 
 #include "Math/Box.h"
+#include "Text3DTypes.h"
 #include "Text3DRendererBase.generated.h"
 
 class UText3DComponent;
 
-enum class EText3DRendererFlags : uint8
-{
-	/** Update whole geometry for text */
-	Geometry = 1 << 0,
-	/** Update layout for characters (transform) */
-	Layout = 1 << 1,
-	/** Update visibility properties */
-	Visibility = 1 << 2,
-	/** Update materials slots */
-	Material = 1 << 3,
-	/** Update everything */
-	All = Geometry | Layout | Visibility | Material
-};
-
 /**
  * Base class for a rendering implementation of Text3D
  * The whole rendering logic should be encapsulated into an instance of this class
@@ -67,6 +54,9 @@ protected:
 
 	UText3DComponent* GetText3DComponent() const;
 
+	/** Recalculates bounds when a layout or geometry change happened */
+	void RefreshBounds();
+
 private:
 	TOptional<FBox> CachedBounds;
 	bool bInitialized = false;
diff --git a/Engine/Plugins/Experimental/Text3D/Source/Text3D/Public/Text3DComponent.h b/Engine/Plugins/Experimental/Text3D/Source/Text3D/Public/Text3DComponent.h
--- a/Engine/Plugins/Experimental/Text3D/Source/Text3D/Public/Text3DComponent.h
+++ b/Engine/Plugins/Experimental/Text3D/Source/Text3D/Public/Text3DComponent.h
@@ -3,6 +3,7 @@
 #pragma once
 
 #include "Components/SceneComponent.h"
+#include "Containers/Ticker.h"
 #include "Text3DTypes.h"
 #include "UObject/ObjectMacros.h"
 #include "Text3DComponent.generated.h"
@@ -11,7 +12,12 @@ class FTextLayout;
 class ITextLayoutMarshaller;
 class UFont;
 class UMaterialInterface;
+class UText3DLayoutExtensionBase;
+class UText3DGeometryExtensionBase;
+class UText3DMaterialExtensionBase;
 class UText3DRendererBase;
+class UText3DRenderingExtensionBase;
+class UText3DTransformExtensionBase;
 struct FText3DShapedGlyphText;
 struct FTypefaceEntry;
 
@@ -36,15 +42,17 @@ public:
 	{
 		return TextGeneratedNativeDelegate;
 	}
-
-	/** Get whether to allow automatic refresh/mesh generation */
-	bool RefreshesOnChange() const
+	
+	DECLARE_MULTICAST_DELEGATE_TwoParams(FTextUpdated, UText3DComponent*, EText3DRendererFlags);
+	FTextUpdated::RegistrationType& OnTextPostUpdate()
 	{
-		return bRefreshOnChange;
+		return TextPostUpdateDelegate;
 	}
 
-	/** Set whether to allow automatic refresh/mesh generation */
-	void SetRefreshOnChange(const bool Value);
+	FTextUpdated::RegistrationType& OnTextPreUpdate()
+	{
+		return TextPreUpdateDelegate;
+	}
 
 	/** Get the text value and signal the primitives to be rebuilt */
 	const FText& GetText() const
@@ -52,10 +60,20 @@ public:
 		return Text;
 	}
 
+	const FText& GetFormattedText() const;
+
 	/** Set the text value and signal the primitives to be rebuilt */
 	UFUNCTION(BlueprintCallable, Category = "Rendering|Components|Text3D")
 	void SetText(const FText& Value);
 
+	UFUNCTION(BlueprintCallable, Category = "Rendering|Components|Text3D")
+	void SetEnforceUpperCase(bool bInEnforceUpperCase);
+
+	bool GetEnforceUpperCase() const
+	{
+		return bEnforceUpperCase;
+	}
+
 	/**
 	 * Returns the Text property, after being formatted by the FormatText virtual function.
 	 * If FormatText is not overriden, the return FText will be the same as the Text property.
@@ -73,231 +91,137 @@ public:
 	UFUNCTION(BlueprintCallable, Category = "Rendering|Components|Text3D")
 	void SetFont(UFont* InFont);
 
-	/** Get whether an outline is applied. */
-	bool HasOutline() const
-	{
-		return bOutline;
-	}
-
 	/** Set whether an outline is applied. */
 	UFUNCTION(BlueprintCallable, Category = "Rendering|Components|Text3D")
 	void SetHasOutline(const bool bValue);
 
-	/** Get the outline width. */
-	float GetOutlineExpand() const
-	{
-		return OutlineExpand;
-	}
-
 	/** Set the outline width. */
 	UFUNCTION(BlueprintCallable, Category = "Rendering|Components|Text3D")
 	void SetOutlineExpand(const float Value);
 
-	/** Get the text extrusion size and signal the primitives to be rebuilt */
-	float GetExtrude() const
-	{
-		return Extrude;
-	}
-
 	/** Set the text extrusion size and signal the primitives to be rebuilt */
 	UFUNCTION(BlueprintCallable, Category = "Rendering|Components|Text3D")
 	void SetExtrude(const float Value);
 
-	/** Get the 3d bevel value */
-	float GetBevel() const
-	{
-		return Bevel;
-	}
-
 	/** Set the 3d bevel value */
 	UFUNCTION(BlueprintCallable, Category = "Rendering|Components|Text3D")
 	void SetBevel(const float Value);
 
-	/** Get the 3d bevel type */
-	EText3DBevelType GetBevelType() const
-	{
-		return BevelType;
-	}
-
 	/** Set the 3d bevel type */
 	UFUNCTION(BlueprintCallable, Category = "Rendering|Components|Text3D")
 	void SetBevelType(const EText3DBevelType Value);
 
-	/** Get the amount of segments that will be used to tessellate the Bevel */
-	int32 GetBevelSegments() const
-	{
-		return BevelSegments;
-	}
-
 	/** Set the amount of segments that will be used to tessellate the Bevel */
 	UFUNCTION(BlueprintCallable, Category = "Rendering|Components|Text3D")
 	void SetBevelSegments(const int32 Value);
 
 	/** Get the text front material */
-	UMaterialInterface* GetFrontMaterial() const
-	{
-		return FrontMaterial;
-	}
+	UMaterialInterface* GetFrontMaterial() const;
 
 	/** Set the text front material */
 	UFUNCTION(BlueprintCallable, Category = "Rendering|Components|Text3D")
 	void SetFrontMaterial(UMaterialInterface* Value);
 
 	/** Get the text bevel material */
-	UMaterialInterface* GetBevelMaterial() const
-	{
-		return BevelMaterial;
-	}
+	UMaterialInterface* GetBevelMaterial() const;
 
 	/** Set the text bevel material */
 	UFUNCTION(BlueprintCallable, Category = "Rendering|Components|Text3D")
 	void SetBevelMaterial(UMaterialInterface* Value);
-
+	
 	/** Get the text extrude material */
-	UMaterialInterface* GetExtrudeMaterial() const
-	{
-		return ExtrudeMaterial;
-	}
+	UMaterialInterface* GetExtrudeMaterial() const;
 
 	/** Set the text extrude material */
 	UFUNCTION(BlueprintCallable, Category = "Rendering|Components|Text3D")
 	void SetExtrudeMaterial(UMaterialInterface* Value);
 
 	/** Get the text back material */
-	UMaterialInterface* GetBackMaterial() const
-	{
-		return BackMaterial;
-	}
+	UMaterialInterface* GetBackMaterial() const;
 
 	/** Set the text back material */
 	UFUNCTION(BlueprintCallable, Category = "Rendering|Components|Text3D")
 	void SetBackMaterial(UMaterialInterface* Value);
 
 	/** Get the kerning value and signal the primitives to be rebuilt */
-	float GetKerning() const
-	{
-		return Kerning;
-	}
+	float GetKerning() const;
 
 	/** Set the kerning value and signal the primitives to be rebuilt */
 	UFUNCTION(BlueprintCallable, Category = "Rendering|Components|Text3D")
 	void SetKerning(const float Value);
 
 	/** Get the line spacing value and signal the primitives to be rebuilt */
-	float GetLineSpacing() const
-	{
-		return LineSpacing;
-	}
+	float GetLineSpacing() const;
 
 	/** Set the line spacing value and signal the primitives to be rebuilt */
 	UFUNCTION(BlueprintCallable, Category = "Rendering|Components|Text3D")
 	void SetLineSpacing(const float Value);
 
 	/** Get the word spacing value and signal the primitives to be rebuilt */
-	float GetWordSpacing() const
-	{
-		return WordSpacing;
-	}
+	float GetWordSpacing() const;
 
 	/** Set the word spacing value and signal the primitives to be rebuilt */
 	UFUNCTION(BlueprintCallable, Category = "Rendering|Components|Text3D")
 	void SetWordSpacing(const float Value);
 
 	/** Get the horizontal alignment value and signal the primitives to be rebuilt */
-	EText3DHorizontalTextAlignment GetHorizontalAlignment() const
-	{
-		return HorizontalAlignment;
-	}
+	EText3DHorizontalTextAlignment GetHorizontalAlignment() const;
 
 	/** Set the horizontal alignment value and signal the primitives to be rebuilt */
 	UFUNCTION(BlueprintCallable, Category = "Rendering|Components|Text3D")
 	void SetHorizontalAlignment(const EText3DHorizontalTextAlignment value);
 
 	/** Get the vertical alignment and signal the primitives to be rebuilt */
-	EText3DVerticalTextAlignment GetVerticalAlignment() const
-	{
-		return VerticalAlignment;
-	}
+	EText3DVerticalTextAlignment GetVerticalAlignment() const;
 
 	/** Set the vertical alignment and signal the primitives to be rebuilt */
 	UFUNCTION(BlueprintCallable, Category = "Rendering|Components|Text3D")
 	void SetVerticalAlignment(const EText3DVerticalTextAlignment value);
 
 	/** Whether a maximum width is specified */
-	bool HasMaxWidth() const
-	{
-		return bHasMaxWidth;
-	}
+	bool HasMaxWidth() const;
 
 	/** Enable / Disable a Maximum Width */
 	UFUNCTION(BlueprintCallable, Category = "Rendering|Components|Text3D")
 	void SetHasMaxWidth(const bool Value);
 
 	/** Get the Maximum Width - If width is larger, mesh will scale down to fit MaxWidth value */
-	float GetMaxWidth() const
-	{
-		return MaxWidth;
-	}
+	float GetMaxWidth() const;
 
 	/** Set the Maximum Width - If width is larger, mesh will scale down to fit MaxWidth value */
 	UFUNCTION(BlueprintCallable, Category = "Rendering|Components|Text3D")
 	void SetMaxWidth(const float Value);
 
 	/** Get the Maximum Width Handling - Whether to wrap before scaling when the text size reaches the max width */
-	EText3DMaxWidthHandling GetMaxWidthHandling() const
-	{
-		return MaxWidthHandling;
-	}
+	EText3DMaxWidthHandling GetMaxWidthHandling() const;
 
 	/** Set the Maximum Width Handling - Whether to wrap before scaling when the text size reaches the max width */
 	UFUNCTION(BlueprintCallable, Category = "Rendering|Components|Text3D")
 	void SetMaxWidthHandling(const EText3DMaxWidthHandling Value);
 
 	/** Whether a maximum height is specified */
-	bool HasMaxHeight() const
-	{
-		return bHasMaxHeight;
-	}
+	bool HasMaxHeight() const;
 
 	/** Enable / Disable a Maximum Height */
 	UFUNCTION(BlueprintCallable, Category = "Rendering|Components|Text3D")
 	void SetHasMaxHeight(const bool Value);
 
 	/** Get the Maximum Height - If height is larger, mesh will scale down to fit MaxHeight value */
-	float GetMaxHeight() const
-	{
-		return MaxHeight;
-	}
+	float GetMaxHeight() const;
 
 	/** Set the Maximum Height - If height is larger, mesh will scale down to fit MaxHeight value */
 	UFUNCTION(BlueprintCallable, Category = "Rendering|Components|Text3D")
 	void SetMaxHeight(const float Value);
 
 	/** Get if the mesh should scale proportionally when Max Width/Height is set */
-	bool ScalesProportionally() const
-	{
-		return bScaleProportionally;
-	}
+	bool ScalesProportionally() const;
 
 	/** Set if the mesh should scale proportionally when Max Width/Height is set */
 	UFUNCTION(BlueprintCallable, Category = "Rendering|Components|Text3D")
 	void SetScaleProportionally(const bool Value);
 
-	/** Freeze mesh rebuild, to avoid unnecessary mesh rebuilds when setting a few properties together */
-	bool IsFrozen() const
-	{
-		return bFreezeBuild;
-	}
-
-	/** Freeze mesh rebuild, to avoid unnecessary mesh rebuilds when setting a few properties together */
-	UFUNCTION(BlueprintCallable, Category = "Rendering|Components|Text3D")
-	void SetFreeze(const bool bFreeze);
-
 	/** Get the value of CastShadow. */
-	bool CastsShadow() const
-	{
-		return bCastShadow;
-	}
+	bool CastsShadow() const;
 
 	/** Set the value of CastShadow. */
 	UFUNCTION(BlueprintCallable, Category = "Rendering|Components|Text3D")
@@ -305,11 +229,10 @@ public:
 
 	/** Get whole text rendered bounds */
 	UFUNCTION(BlueprintCallable, Category = "Rendering|Components|Text3D")
-	void GetBounds(FVector& Origin, FVector& BoxExtent);
+	void GetBounds(FVector& Origin, FVector& BoxExtent) const;
 
-	/** Gets the scale of actual text geometry, taking into account MaxWidth and MaxHeight constraints. This function will NOT return the component scale*/
-	UFUNCTION(BlueprintCallable, Category = "Rendering|Components|Text3D")
-	FVector GetTextScale();
+	/** Get whole text rendered bounds */
+	FBox GetBounds() const;
 
 	/** Get the typeface */
 	FName GetTypeface() const
@@ -327,23 +250,67 @@ public:
 		return TextRendererClass;
 	}
 
-	/** Manually update the geometry, ignoring RefreshOnChange (but still accounting for the Freeze flag) */
-	void Rebuild();
-
 	const FText3DStatistics& GetStatistics() const
 	{
 		return Statistics;
 	}
 
-	void MarkForLayoutUpdate();
+	EText3DRendererFlags GetUpdateFlags() const;
+	void RequestUpdate(EText3DRendererFlags InFlags, bool bInImmediate = false);
 
-	FVector GetLineLocation(int32 LineIndex) const;
-	UMaterialInterface* GetMaterial(const EText3DGroupType Type) const;
-	TSharedPtr<FText3DShapedGlyphText> GetShapedText() const;
 	uint32 GetTypeFaceIndex() const;
 
 	UText3DRendererBase* GetTextRenderer() const;
 
+	UText3DLayoutExtensionBase* GetLayoutExtension() const
+	{
+		return LayoutExtension;
+	}
+
+	template <typename InType UE_REQUIRES(TIsDerivedFrom<InType, UText3DLayoutExtensionBase>::Value)>
+	InType* GetCastedLayoutExtension() const
+	{
+		return Cast<InType>(LayoutExtension);
+	}
+
+	UText3DMaterialExtensionBase* GetMaterialExtension() const
+	{
+		return MaterialExtension;
+	}
+
+	template <typename InType UE_REQUIRES(TIsDerivedFrom<InType, UText3DMaterialExtensionBase>::Value)>
+	InType* GetCastedMaterialExtension() const
+	{
+		return Cast<InType>(MaterialExtension);
+	}
+
+	UText3DGeometryExtensionBase* GetGeometryExtension() const
+	{
+		return GeometryExtension;
+	}
+
+	template <typename InType UE_REQUIRES(TIsDerivedFrom<InType, UText3DGeometryExtensionBase>::Value)>
+	InType* GetCastedGeometryExtension() const
+	{
+		return Cast<InType>(GeometryExtension);
+	}
+	
+	UText3DRenderingExtensionBase* GetRenderingExtension() const
+	{
+		return RenderingExtension;
+	}
+	
+	template <typename InType UE_REQUIRES(TIsDerivedFrom<InType, UText3DRenderingExtensionBase>::Value)>
+	InType* GetCastedRenderingExtension() const
+	{
+		return Cast<InType>(RenderingExtension);
+	}
+
+	TConstArrayView<UText3DTransformExtensionBase*> GetTransformExtensions() const
+	{
+		return TransformExtensions;
+	}
+
 protected:
 	//~ Begin UActorComponent
 	virtual void OnComponentCreated() override;
@@ -353,6 +320,7 @@ protected:
 	//~ End UActorComponent
 	
 	//~ Begin UObject
+	virtual void Serialize(FArchive& InArchive) override;
 	virtual void PostLoad() override;
 	virtual void PostEditImport() override;
 	virtual void PostDuplicate(EDuplicateMode::Type InDuplicateMode) override;
@@ -362,6 +330,16 @@ protected:
 #endif
 	//~ End UObject
 
+#if WITH_EDITOR
+	friend class UText3DEditorFontSubsystem;
+
+	DECLARE_DELEGATE_RetVal_OneParam(UFont*, FOnResolveFontByName, const FString& InFontName)
+	static FOnResolveFontByName OnResolveFontByNameDelegate;
+
+	/** Legacy : used to resolve font by their names */
+	void SetFontByName(const FString& InFontName);
+#endif
+
 	/** Get the type faces supported for the current font */
 	UFUNCTION()
 	TArray<FName> GetTypefaceNames() const;
@@ -377,13 +355,7 @@ protected:
 	 * Override it to customize text formatting in the final geometry, without affecting the Text property.
 	 * Use GetFormattedText() to retrieve a FText with the result of this formatting.
 	 */
-	virtual void FormatText(FText& InOutText) const {}
-
-	/** Clears all generated components and meshes from this component. */
-	void ClearTextMesh();
-
-	/** Can be used to force an internal geometry and/or layout rebuild */
-	void TriggerInternalRebuild(const EText3DModifyFlags InModifyFlags);
+	virtual void FormatText(FText& InOutText) const;
 
 	/** Called when TextRendererClass is changed */
 	void OnTextRendererClassChanged();
@@ -394,17 +366,8 @@ protected:
 	/** Called when font options are changed */
 	void OnFontPropertiesChanged();
 
-	/** Called when geometry options are changed */
-	void OnGeometryPropertiesChanged();
-
-	/** Called when layout options are changed */
-	void OnLayoutPropertiesChanged();
-
-	/** Called when layout handling options are changed */
-	void OnLayoutHandlingPropertiesChanged();
-
-	/** Called when visibility options are changed */
-	void OnVisibilityPropertiesChanged();
+	/** Plan update task on next tick */
+	void ScheduleTextUpdateNextTick();
 
 	/** The renderer class to use to create the text geometry */
 	UPROPERTY(EditAnywhere, Setter, Getter, Category="Text", NoClear, meta=(ShowDisplayNames, AllowPrivateAccess="true"))
@@ -419,166 +382,157 @@ protected:
 	TOptional<FText> FormattedText;
 
 	/** Text font defines the style of rendered characters */
-	UPROPERTY(EditAnywhere, Getter, Setter, Category = "Text", meta = (AllowPrivateAccess = "true"))
+	UPROPERTY(EditAnywhere, Getter, Setter, Category = "Text", meta = (AdvancedFontPicker, AllowPrivateAccess = "true"))
 	TObjectPtr<UFont> Font;
 
 	/** Text font face, subset within font like bold, italic, regular */
 	UPROPERTY(EditAnywhere, Getter, Setter, Category = "Text", meta = (AllowPrivateAccess = "true", GetOptions="GetTypefaceNames"))
 	FName Typeface;
 
-	/** Whether to allow automatic refresh/mesh generation */
-	UPROPERTY(EditAnywhere, Getter = "RefreshesOnChange", Setter = "SetRefreshOnChange", Category = "Text", AdvancedDisplay, meta = (AllowPrivateAccess = "true"))
-	bool bRefreshOnChange = true;
-
-	/** Size of the extrude */
-	UPROPERTY(EditAnywhere, Getter, Setter, Category = "Geometry", meta = (ClampMin = 0, AllowPrivateAccess = "true"))
-	float Extrude = 5.0f;
-
-	/** Size of bevel */
-	UPROPERTY(EditAnywhere, Getter, Setter, Category = "Geometry", meta = (EditCondition = "!bOutline && Extrude > 0", ClampMin = 0, AllowPrivateAccess = "true"))
-	float Bevel = 0.0f;
-
-	/** Bevel Type */
-	UPROPERTY(EditAnywhere, Getter, Setter, Category = "Geometry", meta = (EditCondition = "!bOutline && Extrude > 0", AllowPrivateAccess = "true"))
-	EText3DBevelType BevelType = EText3DBevelType::Convex;
+	/** Whether to force upper case for text */
+	UPROPERTY(EditAnywhere, Getter = "GetEnforceUpperCase", Setter = "SetEnforceUpperCase", Category = "Text", meta = (AllowPrivateAccess="true"))
+	bool bEnforceUpperCase = false;
 
-	/** Bevel Segments (Defines the amount of tesselation for the bevel part) */
-	UPROPERTY(EditAnywhere, Getter, Setter, Category = "Geometry", meta = (EditCondition = "!bOutline && Extrude > 0", ClampMin = 1, ClampMax = 15, AllowPrivateAccess = "true"))
-	int32 BevelSegments = 8;
+	UPROPERTY(VisibleInstanceOnly, Instanced, Category="Layout", meta=(AllowPrivateAccess="true"))
+	TObjectPtr<UText3DLayoutExtensionBase> LayoutExtension;
 
-	/** Generate Outline */
-	UPROPERTY(EditAnywhere, Getter = "HasOutline", Setter = "SetHasOutline", Category = "Geometry", meta = (AllowPrivateAccess = "true"))
-	bool bOutline = false;
+	UPROPERTY(VisibleInstanceOnly, Instanced, Category="Geometry", meta=(AllowPrivateAccess="true"))
+	TObjectPtr<UText3DGeometryExtensionBase> GeometryExtension;
 
-	/** Outline expand/offset amount */
-	UPROPERTY(EditAnywhere, Getter, Setter, Category = "Geometry", meta = (EditCondition = "bOutline", AllowPrivateAccess = "true"))
-	float OutlineExpand = 0.5f;
+	UPROPERTY(VisibleInstanceOnly, Instanced, Category="Material", meta=(AllowPrivateAccess="true"))
+	TObjectPtr<UText3DMaterialExtensionBase> MaterialExtension;
 
-	/** Material for the front part */
-	UPROPERTY(EditAnywhere, Getter, Setter, Category = "Materials", meta = (AllowPrivateAccess = "true"))
-	TObjectPtr<UMaterialInterface> FrontMaterial;
+	UPROPERTY(VisibleInstanceOnly, Instanced, Category="Rendering", meta=(AllowPrivateAccess="true"))
+	TObjectPtr<UText3DRenderingExtensionBase> RenderingExtension;
 
-	/** Material for the bevel part */
-	UPROPERTY(EditAnywhere, Getter, Setter, Category = "Materials", meta = (EditCondition = "!bOutline && Extrude > 0 && Bevel > 0", EditConditionHides, AllowPrivateAccess = "true"))
-	TObjectPtr<UMaterialInterface> BevelMaterial;
+	UPROPERTY(EditInstanceOnly, Instanced, Category="Transform", meta=(AllowPrivateAccess="true"))
+	TArray<TObjectPtr<UText3DTransformExtensionBase>> TransformExtensions;
 
-	/** Material for the extruded part */
-	UPROPERTY(EditAnywhere, Getter, Setter, Category = "Materials", meta = (EditCondition = "Extrude > 0", EditConditionHides, AllowPrivateAccess = "true"))
-	TObjectPtr<UMaterialInterface> ExtrudeMaterial;
-
-	/** Material for the back part */
-	UPROPERTY(EditAnywhere, Getter, Setter, Category = "Materials", meta = (EditCondition = "Extrude > 0", EditConditionHides, AllowPrivateAccess = "true"))
-	TObjectPtr<UMaterialInterface> BackMaterial;
-
-	/** Horizontal text alignment */
-	UPROPERTY(EditAnywhere, Getter, Setter, Category = "Layout", meta = (AllowPrivateAccess = "true"))
-	EText3DHorizontalTextAlignment HorizontalAlignment = EText3DHorizontalTextAlignment::Left;
+private:
+	/** Active renderer in charge of generating text geometry */
+	UPROPERTY()
+	TObjectPtr<UText3DRendererBase> TextRenderer;
 
-	/** Vertical text alignment */
-	UPROPERTY(EditAnywhere, Getter, Setter, Category = "Layout", meta = (AllowPrivateAccess = "true"))
-	EText3DVerticalTextAlignment VerticalAlignment = EText3DVerticalTextAlignment::FirstLine;
+	UPROPERTY(BlueprintAssignable, Category = Events, meta = (AllowPrivateAccess = "true", DisplayName = "On Text Generated"))
+	FTextGenerated TextGeneratedDelegate;
 
-	/** Text kerning */
-	UPROPERTY(EditAnywhere, Getter, Setter, Category = "Layout", meta = (AllowPrivateAccess = "true"))
-	float Kerning = 0.0f;
+	FTextGeneratedNative TextGeneratedNativeDelegate;
 
-	/** Extra line spacing */
-	UPROPERTY(EditAnywhere, Getter, Setter, Category = "Layout", meta = (AllowPrivateAccess = "true"))
-	float LineSpacing = 0.0f;
+	/** Called before text is updated */
+	FTextUpdated TextPreUpdateDelegate;
 
-	/** Extra word spacing */
-	UPROPERTY(EditAnywhere, Setter, Category = "Layout", meta = (AllowPrivateAccess = "true"))
-	float WordSpacing = 0.0f;
+	/** Called after text is updated */
+	FTextUpdated TextPostUpdateDelegate;
 
-	/** Enables a maximum width to the 3D Text */
-	UPROPERTY(EditAnywhere, Getter = "HasMaxWidth", Setter = "SetHasMaxWidth", Category = "Layout", meta = (InlineEditConditionToggle, AllowPrivateAccess = "true"))
-	bool bHasMaxWidth = false;
+	/** Flagged as true when text is being updated */
+	std::atomic<bool> bIsUpdatingText;
 
-	/** Sets a maximum width to the 3D Text */
-	UPROPERTY(EditAnywhere, Getter, Setter, Category = "Layout", meta = (EditCondition = "bHasMaxWidth", ClampMin = 1, AllowPrivateAccess = "true"))
-	float MaxWidth = 500.f;
+	/** Used to determine and selectively perform the the type of rebuild requested. */
+	EText3DRendererFlags UpdateFlags = EText3DRendererFlags::All;
 
-	/** Enables a maximum height to the 3D Text */
-	UPROPERTY(EditAnywhere, Getter = "HasMaxHeight", Setter = "SetHasMaxHeight", Category = "Layout", meta = (InlineEditConditionToggle, AllowPrivateAccess = "true"))
-	bool bHasMaxHeight = false;
+	/** Text statistics cached since last text generation */
+	FText3DStatistics Statistics;
 
-	/** Dictates how to handle the text if it exceeds the max width */
-	UPROPERTY(EditAnywhere, Getter, Setter, Category = "Layout", meta = (HideEditConditionToggle, EditCondition="bHasMaxWidth", AllowPrivateAccess = "true"))
-	EText3DMaxWidthHandling MaxWidthHandling;
+	/** Planned update requested */
+	FTSTicker::FDelegateHandle TextUpdateHandle;
 
-	/** Sets a maximum height to the 3D Text */
-	UPROPERTY(EditAnywhere, Getter, Setter, Category = "Layout", meta = (EditCondition = "bHasMaxHeight", ClampMin = 1, AllowPrivateAccess = "true"))
-	float MaxHeight = 500.0f;
+	void RebuildInternal(bool bCleanCache = false);
 
-	/** Should the mesh scale proportionally when Max Width/Height is set */
-	UPROPERTY(EditAnywhere, Getter = "ScalesProportionally", Setter = "SetScaleProportionally", Category = "Layout", meta = (AllowPrivateAccess = "true"))
-	bool bScaleProportionally = true;
+	void ClearUpdateFlags();
 
-	// Lighting flags
+	bool IsTypefaceAvailable(FName InTypeface) const;
+	TArray<FTypefaceEntry> GetAvailableTypefaces() const;
+	void RefreshTypeface();
 
-	/** Controls whether the text glyphs should cast a shadow or not. */
-	UPROPERTY(EditAnywhere, Getter = "CastsShadow", Setter = "SetCastShadow", Category = "Lighting", meta = (AllowPrivateAccess = "true"))
-	bool bCastShadow = true;
+	void UpdateStatistics();
 
-private:
-	/** Active renderer in charge of generating text geometry */
+#if WITH_EDITORONLY_DATA
+	UE_DEPRECATED(5.6, "Use Material Extension instead")
 	UPROPERTY()
-	TObjectPtr<UText3DRendererBase> TextRenderer;
+	TObjectPtr<UMaterialInterface> FrontMaterial_DEPRECATED;
 
-	UPROPERTY(BlueprintAssignable, Category = Events, meta = (AllowPrivateAccess = "true", DisplayName = "On Text Generated"))
-	FTextGenerated TextGeneratedDelegate;
+	UE_DEPRECATED(5.6, "Use Material Extension instead")
+	UPROPERTY()
+	TObjectPtr<UMaterialInterface> BevelMaterial_DEPRECATED;
 
-	FTextGeneratedNative TextGeneratedNativeDelegate;
+	UE_DEPRECATED(5.6, "Use Material Extension instead")
+	UPROPERTY()
+	TObjectPtr<UMaterialInterface> ExtrudeMaterial_DEPRECATED;
 
-	/** Flagged as true when the text mesh is being built. */
-	std::atomic<bool> bIsBuilding;
+	UE_DEPRECATED(5.6, "Use Material Extension instead")
+	UPROPERTY()
+	TObjectPtr<UMaterialInterface> BackMaterial_DEPRECATED;
+	
+	UE_DEPRECATED(5.6, "Use Geometry Extension instead")
+	UPROPERTY()
+	float Extrude_DEPRECATED = 5.0f;
 
-	/** No mesh or layout rebuilds occur while this flag is true, allowing the user to set numerous properties before rebuilding. */
-	bool bFreezeBuild = false;
+	UE_DEPRECATED(5.6, "Use Geometry Extension instead")
+	UPROPERTY()
+	float Bevel_DEPRECATED = 0.0f;
 
-	/** Used to determine and selectively perform the the type of rebuild requested. */
-	EText3DModifyFlags ModifyFlags = EText3DModifyFlags::All;
+	UE_DEPRECATED(5.6, "Use Geometry Extension instead")
+	UPROPERTY()
+	EText3DBevelType BevelType_DEPRECATED = EText3DBevelType::Convex;
+		
+	UE_DEPRECATED(5.6, "Use Geometry Extension instead")
+	UPROPERTY()
+	int32 BevelSegments_DEPRECATED = 8;
 
-	/** Additional scale to apply to the text. */
-	FVector TextScale = FVector::ZeroVector;
+	UE_DEPRECATED(5.6, "Use Geometry Extension instead")
+	UPROPERTY()
+	bool bOutline_DEPRECATED = false;
 
-	/** Caches the last result of ShapedText, to allow faster updates of layout changes. */
-	TSharedPtr<FText3DShapedGlyphText> ShapedText;
+	UE_DEPRECATED(5.6, "Use Geometry Extension instead")
+	UPROPERTY()
+	float OutlineExpand_DEPRECATED = 0.5f;
 
-	/** Stores the text layout calculated by the TextLayoutMarshaller. */
-	TSharedPtr<FTextLayout> TextLayout;
+	UE_DEPRECATED(5.6, "Use Layout Extension instead")
+	UPROPERTY()
+	EText3DHorizontalTextAlignment HorizontalAlignment_DEPRECATED = EText3DHorizontalTextAlignment::Left;
 
-	/** Determines how text is laid out, ie. parsing line breaks. */
-	TSharedPtr<ITextLayoutMarshaller> TextLayoutMarshaller;
+	UE_DEPRECATED(5.6, "Use Layout Extension instead")
+	UPROPERTY()
+	EText3DVerticalTextAlignment VerticalAlignment_DEPRECATED = EText3DVerticalTextAlignment::FirstLine;
 
-	/** Text statistics cached since last text generation */
-	FText3DStatistics Statistics;
+	UE_DEPRECATED(5.6, "Use Layout Extension instead")
+	UPROPERTY()
+	float Kerning_DEPRECATED = 0.0f;
 
-	/** Slot based Material accessors. */
-	void SetMaterial(const EText3DGroupType Type, UMaterialInterface* Material);
-	void OnMaterialChanged();
+	UE_DEPRECATED(5.6, "Use Layout Extension instead")
+	UPROPERTY()
+	float LineSpacing_DEPRECATED = 0.0f;
 
-	float GetMaxBevel() const;
+	UE_DEPRECATED(5.6, "Use Layout Extension instead")
+	UPROPERTY()
+	float WordSpacing_DEPRECATED = 0.0f;
 
-	/** Selective build functionality (geometry, only layout, etc.) */
-	void UpdateTransforms();
-	void RebuildInternal(const bool& bIsAutoUpdate = true, const bool& bCleanCache = false);
-	void BuildTextMesh(const bool& bCleanCache = false);
-	void BuildTextMeshInternal(const bool& bCleanCache);
+	UE_DEPRECATED(5.6, "Use Layout Extension instead")
+	UPROPERTY()
+	bool bHasMaxWidth_DEPRECATED = false;
 
-	/** Layout functionality. */
-	float GetTextHeight() const;
-	void CalculateTextScale();
+	UE_DEPRECATED(5.6, "Use Layout Extension instead")
+	UPROPERTY()
+	float MaxWidth_DEPRECATED = 500.f;
 
-	/** Convenience functions to set and query ModifyFlags. */
-	bool NeedsMeshRebuild() const;
-	bool NeedsLayoutUpdate() const;
-	void MarkForGeometryUpdate();
-	void ClearUpdateFlags();
+	UE_DEPRECATED(5.6, "Use Layout Extension instead")
+	UPROPERTY()
+	bool bHasMaxHeight_DEPRECATED = false;
 
-	bool IsTypefaceAvailable(FName InTypeface) const;
-	TArray<FTypefaceEntry> GetAvailableTypefaces() const;
-	void RefreshTypeface();
+	UE_DEPRECATED(5.6, "Use Layout Extension instead")
+	UPROPERTY()
+	EText3DMaxWidthHandling MaxWidthHandling_DEPRECATED = EText3DMaxWidthHandling::Scale;
+	
+	UE_DEPRECATED(5.6, "Use Layout Extension instead")
+	UPROPERTY()
+	float MaxHeight_DEPRECATED = 500.0f;
 
-	void UpdateStatistics();
-};
\ No newline at end of file
+	UE_DEPRECATED(5.6, "Use Layout Extension instead")
+	UPROPERTY()
+	bool bScaleProportionally_DEPRECATED = true;
+	
+	UE_DEPRECATED(5.6, "Use Rendering Extension instead")
+	UPROPERTY()
+	bool bCastShadow_DEPRECATED = true;
+#endif
+};
diff --git a/Engine/Plugins/Experimental/Text3D/Source/Text3D/Public/Text3DTypes.h b/Engine/Plugins/Experimental/Text3D/Source/Text3D/Public/Text3DTypes.h
--- a/Engine/Plugins/Experimental/Text3D/Source/Text3D/Public/Text3DTypes.h
+++ b/Engine/Plugins/Experimental/Text3D/Source/Text3D/Public/Text3DTypes.h
@@ -2,8 +2,52 @@
 
 #pragma once
 
+#include "Misc/EnumClassFlags.h"
 #include "Text3DTypes.generated.h"
 
+/** Enumerate Text3D update flags based on their priority */
+enum class EText3DRendererFlags : uint8
+{
+	None,
+	/** Update whole geometry for text */
+	Geometry = 1 << 0,
+	/** Update layout for characters (transform) */
+	Layout = 1 << 1,
+	/** Update materials slots */
+	Material = 1 << 2,
+	/** Update visibility/lighting properties */
+	Visibility = 1 << 3,
+	/** Update everything */
+	All = Geometry | Layout | Visibility | Material
+};
+
+UENUM()
+enum class EText3DMaterialStyle : uint8 
+{
+	Invalid UMETA(Hidden),
+	Solid,
+	Gradient,
+	Texture,
+	Custom
+};
+
+UENUM()
+enum class EText3DMaterialBlendMode : uint8
+{
+	Invalid UMETA(Hidden),
+	Opaque,
+	Translucent
+};
+
+UENUM()
+enum class EText3DFontStyleFlags : uint8
+{
+	None = 0,
+	Monospace = 1 << 0,
+	Bold = 1 << 1,
+	Italic = 1 << 2,
+};
+
 UENUM()
 enum class EText3DBevelType : uint8
 {
@@ -52,18 +96,6 @@ enum class EText3DMaxWidthHandling : uint8
 	WrapAndScale	UMETA(DisplayName = "Wrap and Scale"),
 };
 
-UENUM()
-enum class EText3DModifyFlags : uint8
-{
-	None = 0,
-	Layout = 1 << 0,
-	Geometry = 1 << 1,
-	Unfreeze = 1 << 2,
-
-	All = Layout | Geometry | Unfreeze
-};
-ENUM_CLASS_FLAGS(EText3DModifyFlags)
-
 UENUM()
 enum class EText3DCharacterEffectOrder : uint8
 {
@@ -86,4 +118,96 @@ struct FText3DStatistics
 {
 	TArray<FText3DWordStatistics> Words;
 	int32 WhiteSpaces;
+};
+
+struct FText3DFontFamily
+{
+	void AddFontFace(const FString& InFontFaceName, const FString& InFontFacePath)
+	{
+		if (FontFacePaths.Contains(InFontFaceName))
+		{
+			return;
+		}
+
+		FontFacePaths.Add(InFontFaceName, InFontFacePath);
+	}
+
+	/** Family these font faces belong to */
+	FString FontFamilyName;
+
+	/** Map of each font face with name -> path */
+	TMap<FString, FString> FontFacePaths;
+};
+
+/** Used to identify a specific material type */
+USTRUCT()
+struct FText3DMaterialKey
+{
+	GENERATED_BODY()
+
+	FText3DMaterialKey() = default;
+	explicit FText3DMaterialKey(EText3DMaterialBlendMode InBlend, bool bInIsUnlit)
+		: BlendMode(InBlend)
+		, bIsUnlit(bInIsUnlit)
+	{}
+		
+	bool operator==(const FText3DMaterialKey& Other) const
+	{
+		return BlendMode == Other.BlendMode
+			&& bIsUnlit == Other.bIsUnlit;
+	}
+
+	bool operator!=(const FText3DMaterialKey& Other) const
+	{
+		return !(*this == Other);
+	}
+	
+	friend uint32 GetTypeHash(const FText3DMaterialKey& InMaterialSettings)
+	{
+		return HashCombineFast(GetTypeHash(InMaterialSettings.BlendMode), GetTypeHash(InMaterialSettings.bIsUnlit));
+	}
+
+	UPROPERTY()
+	EText3DMaterialBlendMode BlendMode = EText3DMaterialBlendMode::Invalid;
+
+	UPROPERTY()
+	bool bIsUnlit = false;
+};
+
+USTRUCT()
+struct FText3DMaterialGroupKey
+{
+	GENERATED_BODY()
+
+	FText3DMaterialGroupKey() = default;
+	explicit FText3DMaterialGroupKey(FText3DMaterialKey InKey, EText3DGroupType InGroup, EText3DMaterialStyle InStyle)
+		: Key(InKey)
+		, Group(InGroup)
+		, Style(InStyle)
+	{}
+
+	bool operator==(const FText3DMaterialGroupKey& Other) const
+	{
+		return Key == Other.Key
+			&& Group == Other.Group
+			&& Style == Other.Style;
+	}
+
+	bool operator!=(const FText3DMaterialGroupKey& Other) const
+	{
+		return !(*this == Other);
+	}
+
+	friend uint32 GetTypeHash(const FText3DMaterialGroupKey& InMaterialSettings)
+	{
+		return HashCombineFast(HashCombineFast(GetTypeHash(InMaterialSettings.Key), GetTypeHash(InMaterialSettings.Group)), GetTypeHash(InMaterialSettings.Style));
+	}
+
+	UPROPERTY()
+	FText3DMaterialKey Key;
+
+	EText3DGroupType Group = EText3DGroupType::Front;
+
+	UPROPERTY()
+	EText3DMaterialStyle Style = EText3DMaterialStyle::Invalid;
 };
\ No newline at end of file
diff --git a/Engine/Plugins/Experimental/Text3D/Source/Text3D/Text3D.Build.cs b/Engine/Plugins/Experimental/Text3D/Source/Text3D/Text3D.Build.cs
--- a/Engine/Plugins/Experimental/Text3D/Source/Text3D/Text3D.Build.cs
+++ b/Engine/Plugins/Experimental/Text3D/Source/Text3D/Text3D.Build.cs
@@ -10,13 +10,16 @@ public class Text3D : ModuleRules
 		bEnableExceptions = true;
 		PCHUsage = ModuleRules.PCHUsageMode.UseExplicitOrSharedPCHs;
 
-		PublicDependencyModuleNames.AddRange(new string[] {
+		PublicDependencyModuleNames.AddRange(new string[] 
+		{
 			"Core",
 			"CoreUObject",
 			"Engine",
 		});
 
-		PrivateDependencyModuleNames.AddRange(new string[] {
+		PrivateDependencyModuleNames.AddRange(new string[] 
+		{
+			"DeveloperSettings",
 			"FreeType2",
 			"GeometryAlgorithms",
 			"GeometryCore",
diff --git a/Engine/Plugins/Experimental/Text3D/Source/Text3DEditor/Private/Styles/Text3DEditorStyle.cpp b/Engine/Plugins/Experimental/Text3D/Source/Text3DEditor/Private/Styles/Text3DEditorStyle.cpp
--- a/Engine/Plugins/Experimental/Text3D/Source/Text3DEditor/Private/Styles/Text3DEditorStyle.cpp
+++ b/Engine/Plugins/Experimental/Text3D/Source/Text3DEditor/Private/Styles/Text3DEditorStyle.cpp
@@ -17,9 +17,18 @@ FText3DEditorStyle::FText3DEditorStyle()
 
 	ContentRootDir = FPaths::Combine(Plugin->GetBaseDir(), TEXT("Resources"));
 
-	Set("ClassIcon.Text3DActor", new IMAGE_BRUSH("Text3DActor_16x", FVector2D (16.0f, 16.0f)));
+	// ClassIcon
+	Set("ClassIcon.Text3DActor", new IMAGE_BRUSH("Text3DActor_16x", Icon16x16));
 	Set("ClassThumbnail.Text3DActor", new IMAGE_BRUSH("Text3DActor_64x", FVector2D (64.0f, 64.0f)));
 
+	// Alignments
+	Set("Icons.Alignment.Left", new IMAGE_BRUSH_SVG("AlignLeft", Icon16x16));
+	Set("Icons.Alignment.Center_Y", new IMAGE_BRUSH_SVG("AlignCenterHoriz", Icon16x16));
+	Set("Icons.Alignment.Right", new IMAGE_BRUSH_SVG("AlignRight", Icon16x16));
+	Set("Icons.Alignment.Top", new IMAGE_BRUSH_SVG("AlignTop", Icon16x16));
+	Set("Icons.Alignment.Center_Z", new IMAGE_BRUSH_SVG("AlignCenterVert", Icon16x16));
+	Set("Icons.Alignment.Bottom", new IMAGE_BRUSH_SVG("AlignBottom", Icon16x16));
+
 	FSlateStyleRegistry::RegisterSlateStyle(*this);
 }
 
diff --git a/Engine/Plugins/Experimental/Text3D/Source/Text3DEditor/Private/Text3DEditorModule.cpp b/Engine/Plugins/Experimental/Text3D/Source/Text3DEditor/Private/Text3DEditorModule.cpp
--- a/Engine/Plugins/Experimental/Text3D/Source/Text3DEditor/Private/Text3DEditorModule.cpp
+++ b/Engine/Plugins/Experimental/Text3D/Source/Text3DEditor/Private/Text3DEditorModule.cpp
@@ -1,8 +1,20 @@
 // Copyright Epic Games, Inc. All Rights Reserved.
 
 #include "Text3DEditorModule.h"
+
+#include "Commands/Text3DEditorFontSelectorCommands.h"
+#include "Customizations/Text3DEditorFontPropertyTypeCustomization.h"
+#include "Customizations/Text3DEditorHorizontalPropertyTypeCustomization.h"
+#include "Customizations/Text3DEditorVerticalPropertyTypeCustomization.h"
+#include "Engine/Font.h"
+#include "Logs/Text3DEditorLogs.h"
 #include "Modules/ModuleManager.h"
+#include "PropertyEditorModule.h"
 #include "Styles/Text3DEditorStyle.h"
+#include "Text3DTypes.h"
+#include "UObject/ReflectedTypeAccessors.h"
+
+DEFINE_LOG_CATEGORY(LogText3DEditor);
 
 IMPLEMENT_MODULE(FText3DEditorModule, Text3DEditor)
 
@@ -10,4 +22,32 @@ void FText3DEditorModule::StartupModule()
 {
 	// Init style set
 	FText3DEditorStyle::Get();
+
+	FText3DEditorFontSelectorCommands::Register();
+	
+	FPropertyEditorModule& PropertyModule = FModuleManager::GetModuleChecked<FPropertyEditorModule>(TEXT("PropertyEditor"));
+
+	PropertyModule.RegisterCustomPropertyTypeLayout(RegisteredTypeNames.Add_GetRef(StaticEnum<EText3DHorizontalTextAlignment>()->GetFName())
+		, FOnGetPropertyTypeCustomizationInstance::CreateStatic(&FText3DEditorHorizontalPropertyTypeCustomization::MakeInstance));
+
+	PropertyModule.RegisterCustomPropertyTypeLayout(RegisteredTypeNames.Add_GetRef(StaticEnum<EText3DVerticalTextAlignment>()->GetFName())
+		, FOnGetPropertyTypeCustomizationInstance::CreateStatic(&FText3DEditorVerticalPropertyTypeCustomization::MakeInstance));
+	
+	PropertyModule.RegisterCustomPropertyTypeLayout(RegisteredTypeNames.Add_GetRef(UFont::StaticClass()->GetFName())
+		, FOnGetPropertyTypeCustomizationInstance::CreateStatic(&FText3DEditorFontPropertyTypeCustomization::MakeInstance), MakeShared<FText3DEditorFontPropertyTypeIdentifier>());
+}
+
+void FText3DEditorModule::ShutdownModule()
+{
+	FText3DEditorFontSelectorCommands::Unregister();
+
+	if (FPropertyEditorModule* PropertyModule = FModuleManager::GetModulePtr<FPropertyEditorModule>(TEXT("PropertyEditor")))
+	{
+		for (const FName& RegisteredTypeName : RegisteredTypeNames)
+		{
+			PropertyModule->UnregisterCustomPropertyTypeLayout(RegisteredTypeName);
+		}
+	}
+
+	RegisteredTypeNames.Empty();
 }
diff --git a/Engine/Plugins/Experimental/Text3D/Source/Text3DEditor/Private/Text3DEditorModule.h b/Engine/Plugins/Experimental/Text3D/Source/Text3DEditor/Private/Text3DEditorModule.h
--- a/Engine/Plugins/Experimental/Text3D/Source/Text3DEditor/Private/Text3DEditorModule.h
+++ b/Engine/Plugins/Experimental/Text3D/Source/Text3DEditor/Private/Text3DEditorModule.h
@@ -1,7 +1,10 @@
 // Copyright Epic Games, Inc. All Rights Reserved.
+
 #pragma once
 
+#include "Containers/Array.h"
 #include "Modules/ModuleInterface.h"
+#include "UObject/NameTypes.h"
 
 class FText3DEditorModule final : public IModuleInterface
 {
@@ -10,5 +13,9 @@ public:
 
 	//~ Begin IModuleInterface
 	virtual void StartupModule() override;
+	virtual void ShutdownModule() override;
 	//~ End IModuleInterface
+
+private:
+	TArray<FName> RegisteredTypeNames;
 };
\ No newline at end of file
diff --git a/Engine/Plugins/Experimental/Text3D/Source/Text3DEditor/Text3DEditor.Build.cs b/Engine/Plugins/Experimental/Text3D/Source/Text3DEditor/Text3DEditor.Build.cs
--- a/Engine/Plugins/Experimental/Text3D/Source/Text3DEditor/Text3DEditor.Build.cs
+++ b/Engine/Plugins/Experimental/Text3D/Source/Text3DEditor/Text3DEditor.Build.cs
@@ -1,5 +1,6 @@
 // Copyright Epic Games, Inc. All Rights Reserved.
 
+using System.IO;
 using UnrealBuildTool;
 
 public class Text3DEditor : ModuleRules
@@ -9,12 +10,38 @@ public class Text3DEditor : ModuleRules
 		PrivateDependencyModuleNames.AddRange(new string[] {
 			"Core",
 			"CoreUObject",
+			"DeveloperSettings",
 			"EditorFramework",
+			"EditorSubsystem",
+			"Engine",
+			"FreeType2",
+			"InputCore",
 			"Projects",
 			"SlateCore",
 			"Slate",
 			"Text3D",
+			"ToolMenus",
 			"UnrealEd",
 		});
+
+		if (Target.Platform.IsInGroup(UnrealPlatformGroup.Windows))
+		{
+			PrivateDependencyModuleNames.Add("DirectX");
+
+			PublicAdditionalLibraries.AddRange(new string[] {
+				Path.Combine(Target.WindowsPlatform.DirectXLibDir, "dwrite.lib")
+			});
+		}
+
+		// Needed for FreeType/Font functionality
+		AddEngineThirdPartyPrivateStaticDependencies(Target,
+			"zlib",
+			"UElibPNG"
+		);
+
+		if (Target.Platform == UnrealTargetPlatform.Mac)
+		{
+			PublicFrameworks.AddRange(new string[] { "Foundation", "CoreText", "ApplicationServices"});
+		}
 	}
 }
