MotionDesign : Text3D

diff --git a/Engine/Plugins/Experimental/Text3D/Text3D.uplugin b/Engine/Plugins/Experimental/Text3D/Text3D.uplugin
--- a/Engine/Plugins/Experimental/Text3D/Text3D.uplugin
+++ b/Engine/Plugins/Experimental/Text3D/Text3D.uplugin
@@ -2,14 +2,14 @@
   "FileVersion": 3,
   "Version": 1,
   "FriendlyName": "Text 3D",
-  "Description": "3D Text Generator",
+  "Description": "Tool to create 3D Text with advanced options",
   "Category": "Text",
   "CreatedBy": "Epic Games, Inc.",
   "CreatedByURL": "https://epicgames.com",
   "DocsURL": "",
   "MarketplaceURL": "",
   "SupportURL": "",
-  "CanContainContent": false,
+  "CanContainContent": true,
   "IsBetaVersion": true,
   "Installed": false,
   "Modules": [
