MotionDesign : Text3D
- Fix non unity horde issue 861573 with mac font system

diff --git a/Engine/Plugins/Experimental/Text3D/Source/Text3DEditor/Private/Platforms/Mac/MacSystemFontLoading.cpp b/Engine/Plugins/Experimental/Text3D/Source/Text3DEditor/Private/Platforms/Mac/MacSystemFontLoading.cpp
--- a/Engine/Plugins/Experimental/Text3D/Source/Text3DEditor/Private/Platforms/Mac/MacSystemFontLoading.cpp
+++ b/Engine/Plugins/Experimental/Text3D/Source/Text3DEditor/Private/Platforms/Mac/MacSystemFontLoading.cpp
@@ -6,6 +6,7 @@
 #include <CoreText/CoreText.h>
 #include <Foundation/Foundation.h>
 #include "Logs/Text3DEditorLogs.h"
+#include "Misc/Paths.h"
 #include "Text3DTypes.h"
 
 namespace UE::Text3D::Private::Fonts
