MotionDesign : Text3D
- Add individual character properties (kerning, transform)
- Add character pivot alignment for animation and adjusting layout to account for this
- Max width improvements to only trigger a layout update instead of geometry and layout
- Improving overall layout code

diff --git a/Engine/Plugins/Experimental/Avalanche/Source/AvalancheText/Private/AvaTextActor.cpp b/Engine/Plugins/Experimental/Avalanche/Source/AvalancheText/Private/AvaTextActor.cpp
--- a/Engine/Plugins/Experimental/Avalanche/Source/AvalancheText/Private/AvaTextActor.cpp
+++ b/Engine/Plugins/Experimental/Avalanche/Source/AvalancheText/Private/AvaTextActor.cpp
@@ -75,3 +75,10 @@ void AAvaTextActor::SetColorData(const FAvaColorChangeData& NewColorData)
 		}
 	}
 }
+
+#if WITH_EDITOR
+FString AAvaTextActor::GetDefaultActorLabel() const
+{
+	return TEXT("Text3DActor");
+}
+#endif
diff --git a/Engine/Plugins/Experimental/Avalanche/Source/AvalancheText/Public/AvaText3DComponent.h b/Engine/Plugins/Experimental/Avalanche/Source/AvalancheText/Public/AvaText3DComponent.h
--- a/Engine/Plugins/Experimental/Avalanche/Source/AvalancheText/Public/AvaText3DComponent.h
+++ b/Engine/Plugins/Experimental/Avalanche/Source/AvalancheText/Public/AvaText3DComponent.h
@@ -7,7 +7,11 @@
 #include "Text3DComponent.h"
 #include "AvaText3DComponent.generated.h"
 
-UCLASS(MinimalAPI, ClassGroup="Text3D", PrioritizeCategories=("Text", "Layout", "Geometry", "Style", "Materials"), meta=(BlueprintSpawnableComponent))
+UCLASS(MinimalAPI
+	, ClassGroup="Text3D"
+	, PrioritizeCategories="Text Character Layout Geometry Material Rendering Transform"
+	, HideCategories=(Replication,Collision,HLOD,Physics,Networking,Input,Actor,Cooking,LevelInstance,Streaming,DataLayers,WorldPartition)
+	, meta=(BlueprintSpawnableComponent))
 class UAvaText3DComponent : public UText3DComponent
 {
 	friend class FAvaTextVisualizer;
diff --git a/Engine/Plugins/Experimental/Avalanche/Source/AvalancheText/Public/AvaTextActor.h b/Engine/Plugins/Experimental/Avalanche/Source/AvalancheText/Public/AvaTextActor.h
--- a/Engine/Plugins/Experimental/Avalanche/Source/AvalancheText/Public/AvaTextActor.h
+++ b/Engine/Plugins/Experimental/Avalanche/Source/AvalancheText/Public/AvaTextActor.h
@@ -25,6 +25,12 @@ public:
 	AVALANCHETEXT_API FAvaColorChangeData GetColorData() const;
 	AVALANCHETEXT_API void SetColorData(const FAvaColorChangeData& InColorData);
 
+	//~ Begin AActor
+#if WITH_EDITOR
+	virtual FString GetDefaultActorLabel() const override;
+#endif
+	//~ End AActor
+
 private:
 	UPROPERTY(Category = "Text", VisibleAnywhere, BlueprintReadOnly, meta = (AllowPrivateAccess = "true"))
 	TObjectPtr<UText3DComponent> Text3DComponent;
diff --git a/Engine/Plugins/Experimental/Avalanche/Source/AvalancheTextEditor/Private/Visualizer/AvaTextVisualizer.cpp b/Engine/Plugins/Experimental/Avalanche/Source/AvalancheTextEditor/Private/Visualizer/AvaTextVisualizer.cpp
--- a/Engine/Plugins/Experimental/Avalanche/Source/AvalancheTextEditor/Private/Visualizer/AvaTextVisualizer.cpp
+++ b/Engine/Plugins/Experimental/Avalanche/Source/AvalancheTextEditor/Private/Visualizer/AvaTextVisualizer.cpp
@@ -1,9 +1,11 @@
 // Copyright Epic Games, Inc. All Rights Reserved.
 
 #include "AvaTextVisualizer.h"
+
 #include "AvaField.h"
 #include "AvaShapeSprites.h"
 #include "AvaText3DComponent.h"
+#include "Characters/Text3DDefaultCharacter.h"
 #include "EditorViewportClient.h"
 #include "Engine/Texture2D.h"
 #include "Extensions/Text3DDefaultLayoutExtension.h"
@@ -27,6 +29,7 @@ IMPLEMENT_HIT_PROXY(HAvaTextGradientLineStartHandleProxy, HAvaHitProxy)
 IMPLEMENT_HIT_PROXY(HAvaTextGradientLineEndHandleProxy, HAvaHitProxy)
 IMPLEMENT_HIT_PROXY(HAvaTextGradientCenterHandleProxy, HAvaHitProxy)
 IMPLEMENT_HIT_PROXY(HAvaTextGradientSmoothnessHandleProxy, HAvaHitProxy)
+IMPLEMENT_HIT_PROXY(HAvaTextCharacterKerningHandleProxy, HAvaHitProxy)
 
 namespace UE::Ava::Private
 {
@@ -52,6 +55,7 @@ FAvaTextVisualizer::FAvaTextVisualizer()
 	MaxHeightProperty           = GetProperty<UText3DDefaultLayoutExtension>(UText3DDefaultLayoutExtension::GetMaxHeightPropertyName());
 	MaxWidthProperty            = GetProperty<UText3DDefaultLayoutExtension>(UText3DDefaultLayoutExtension::GetMaxWidthPropertyName());
 	ScaleProportionallyProperty = GetProperty<UText3DDefaultLayoutExtension>(UText3DDefaultLayoutExtension::GetScaleProportionallyPropertyName());
+	CharacterKerningProperty    = GetProperty<UText3DDefaultCharacter>(UText3DDefaultCharacter::GetKerningPropertyName());
 
 	ResetEditingFlags();
 	bShowGradientControls = false;
@@ -136,8 +140,7 @@ bool FAvaTextVisualizer::VisProxyHandleClick(FEditorViewportClient* InViewportCl
 			return true;
 		}
 	}
-
-	if (VisProxy->IsA(HAvaTextMaxTextWidthProxy::StaticGetType()))
+	else if (VisProxy->IsA(HAvaTextMaxTextWidthProxy::StaticGetType()))
 	{
 		if (const HAvaTextMaxTextWidthProxy* const MaxWidthProxy = static_cast<HAvaTextMaxTextWidthProxy*>(VisProxy))
 		{
@@ -205,9 +208,8 @@ bool FAvaTextVisualizer::VisProxyHandleClick(FEditorViewportClient* InViewportCl
 			return true;
 		}
 	}
-
 	// this one just toggles gradient control on/off, in case the actor has it
-	if (VisProxy->IsA(HAvaTextEditGradientProxy::StaticGetType()))
+	else if (VisProxy->IsA(HAvaTextEditGradientProxy::StaticGetType()))
 	{
 		EndEditing();			
 		bShowGradientControls = !bShowGradientControls;
@@ -281,6 +283,21 @@ bool FAvaTextVisualizer::VisProxyHandleClick(FEditorViewportClient* InViewportCl
 			}			
 		}
 	}
+	else if (VisProxy->IsA(HAvaTextCharacterKerningHandleProxy::StaticGetType()))
+	{
+		if (const HAvaTextCharacterKerningHandleProxy* const KerningHandleProxy = static_cast<HAvaTextCharacterKerningHandleProxy*>(VisProxy))
+		{
+			EndEditing();
+
+			if (UAvaText3DComponent* AvaTextComponent = Cast<UAvaText3DComponent>(const_cast<UActorComponent*>(KerningHandleProxy->Component.Get())))
+			{
+				TextComponent = AvaTextComponent;
+				EditingKerningIndex = KerningHandleProxy->Index;
+				StartEditing(InViewportClient, AvaTextComponent);
+				return true;
+			}			
+		}
+	}
 
 	return false;
 }
@@ -330,7 +347,13 @@ bool FAvaTextVisualizer::GetWidgetLocation(const FEditorViewportClient* Viewport
 		{
 			OutLocation = GetGradientSmoothnessHandleLocation(AvaTextComponent);
 			return true;
-		}		
+		}
+
+		if (EditingKerningIndex != INDEX_NONE)
+		{
+			OutLocation = GetKerningHandleLocation(AvaTextComponent, EditingKerningIndex);
+			return true;
+		}
 	}
 
 	return Super::GetWidgetLocation(ViewportClient, OutLocation);
@@ -362,6 +385,12 @@ bool FAvaTextVisualizer::GetWidgetMode(const FEditorViewportClient* ViewportClie
 		return true;
 	}
 
+	if (EditingKerningIndex != INDEX_NONE)
+	{
+		Mode = UE::Widget::WM_Translate;
+		return true;
+	}
+
 	return Super::GetWidgetMode(ViewportClient, Mode);
 }
 
@@ -397,6 +426,12 @@ bool FAvaTextVisualizer::GetWidgetAxisList(const FEditorViewportClient* Viewport
 		AxisList = EAxisList::Type::Z;
 		return true;
 	}
+	
+	if (EditingKerningIndex != INDEX_NONE)
+	{
+		AxisList = EAxisList::Y;
+		return true;
+	}
 
 	return Super::GetWidgetAxisList(ViewportClient, WidgetMode, AxisList);
 }
@@ -471,8 +506,7 @@ bool FAvaTextVisualizer::HandleInputDeltaInternal(FEditorViewportClient* InViewp
 		EndEditing();
 		return true;
 	}
-
-	if (bEditingHeight)
+	else if (bEditingHeight)
 	{
 		if (UAvaText3DComponent* const AvaTextComponent = Cast<UAvaText3DComponent>(GetEditedComponent()))
 		{
@@ -523,8 +557,7 @@ bool FAvaTextVisualizer::HandleInputDeltaInternal(FEditorViewportClient* InViewp
 		EndEditing();
 		return true;
 	}
-
-	if (bEditingGradientRotation_EndHandle || bEditingGradientRotation_StartHandle)
+	else if (bEditingGradientRotation_EndHandle || bEditingGradientRotation_StartHandle)
 	{
 		if (UAvaText3DComponent* const AvaTextComponent = Cast<UAvaText3DComponent>(GetEditedComponent()))
 		{
@@ -571,8 +604,7 @@ bool FAvaTextVisualizer::HandleInputDeltaInternal(FEditorViewportClient* InViewp
 		EndEditing();
 		return true;
 	}
-
-	if (bEditingGradientOffset)
+	else if (bEditingGradientOffset)
 	{
 		if (UAvaText3DComponent* const AvaTextComponent = Cast<UAvaText3DComponent>(GetEditedComponent()))
 		{
@@ -585,6 +617,7 @@ bool FAvaTextVisualizer::HandleInputDeltaInternal(FEditorViewportClient* InViewp
 
 					if (Text3DActor && MaterialExtension)
 					{
+						MaterialExtension->SetFlags(RF_Transactional);
 						MaterialExtension->Modify();
 
 						float GradientOffset = InitialGradientOffset;
@@ -614,8 +647,7 @@ bool FAvaTextVisualizer::HandleInputDeltaInternal(FEditorViewportClient* InViewp
 		EndEditing();
 		return true;
 	}
-
-	if (bEditingGradientSmoothness)
+	else if (bEditingGradientSmoothness)
 	{
 		if (UAvaText3DComponent* const AvaTextComponent = Cast<UAvaText3DComponent>(GetEditedComponent()))
 		{
@@ -625,6 +657,7 @@ bool FAvaTextVisualizer::HandleInputDeltaInternal(FEditorViewportClient* InViewp
 				{
 					if (UText3DDefaultMaterialExtension* MaterialExtension = AvaTextComponent->GetCastedMaterialExtension<UText3DDefaultMaterialExtension>())
 					{
+						MaterialExtension->SetFlags(RF_Transactional);
 						MaterialExtension->Modify();
 
 						float GradientSmoothness = InitialGradientSmoothness;
@@ -648,6 +681,37 @@ bool FAvaTextVisualizer::HandleInputDeltaInternal(FEditorViewportClient* InViewp
 		EndEditing();
 		return true;
 	}
+	else if (EditingKerningIndex != INDEX_NONE)
+	{
+		if (const UAvaText3DComponent* const AvaTextComponent = Cast<UAvaText3DComponent>(GetEditedComponent()))
+		{
+			if (GetViewportWidgetMode(InViewportClient) == UE::Widget::WM_Translate)
+			{
+				if (GetViewportWidgetAxisList(InViewportClient) & EAxisList::Y)
+				{
+					if (UText3DDefaultCharacter* Character = AvaTextComponent->GetCastedCharacter<UText3DDefaultCharacter>(EditingKerningIndex))
+					{
+						Character->SetFlags(RF_Transactional);
+						Character->Modify();
+
+						const float Kerning = InitialCharacterKerning + InAccumulatedTranslation.Y;
+
+						if (!FMath::IsNearlyEqual(Kerning, InitialCharacterKerning))
+						{
+							bHasBeenModified = true;
+						}
+
+						Character->SetKerningOverride(true);
+						Character->SetKerning(Kerning);
+						return true;
+					}
+				}
+			}
+		}
+
+		EndEditing();
+		return true;
+	}
 
 	return Super::HandleInputDeltaInternal(InViewportClient, InViewport, InAccumulatedTranslation, InAccumulatedRotation, InAccumulatedScale);
 }
@@ -716,14 +780,32 @@ bool FAvaTextVisualizer::ResetValue(FEditorViewportClient* InViewportClient, HHi
 
 		FAvaVisualizerBase::NotifyPropertyModified(AvaTextComponent, ScaleProportionallyProperty, EPropertyChangeType::ValueSet);
 	}
+	else if (HitProxy->IsA(HAvaTextCharacterKerningHandleProxy::StaticGetType()))
+	{
+		FScopedTransaction Transaction(NSLOCTEXT("AvaTextVisualizer", "VisualizerResetValue", "Visualizer Reset Value"));
+
+		if (UText3DDefaultCharacter* Character = AvaTextComponent->GetCastedCharacter<UText3DDefaultCharacter>(EditingKerningIndex))
+		{
+			Character->Modify();
+			const UText3DDefaultCharacter* CDO = GetDefault<UText3DDefaultCharacter>();
+			Character->SetKerningOverride(CDO->GetKerningOverride());
+			Character->SetKerning(CDO->GetKerning());
+			FAvaVisualizerBase::NotifyPropertyModified(Character, CharacterKerningProperty, EPropertyChangeType::ValueSet);
+		}
+	}
 
 	return true;
 }
 
 bool FAvaTextVisualizer::IsEditing() const
 {
-	if (bEditingWidth || bEditingHeight ||
-		bEditingGradientRotation_StartHandle || bEditingGradientRotation_EndHandle || bEditingGradientOffset || bEditingGradientSmoothness)
+	if (bEditingWidth
+		|| bEditingHeight 
+		|| bEditingGradientRotation_StartHandle
+		|| bEditingGradientRotation_EndHandle
+		|| bEditingGradientOffset
+		|| bEditingGradientSmoothness
+		|| EditingKerningIndex != INDEX_NONE)
 	{
 		return true;
 	}
@@ -739,9 +821,9 @@ void FAvaTextVisualizer::ResetEditingFlags()
 	bEditingGradientRotation_EndHandle = false;
 	bEditingGradientOffset = false;
 	bEditingGradientSmoothness = false;
+	EditingKerningIndex = INDEX_NONE;
 }
 
-
 void FAvaTextVisualizer::EndEditing()
 {
 	Super::EndEditing();
@@ -885,6 +967,26 @@ FVector FAvaTextVisualizer::GetGradientStartHandleLocation(const UAvaText3DCompo
 	return FVector::Zero();
 }
 
+FVector FAvaTextVisualizer::GetKerningHandleLocation(const UAvaText3DComponent* InTextComponent, int32 InIndex) const
+{
+	if (IsValid(InTextComponent))
+	{
+		if (UText3DCharacterBase* Character = InTextComponent->GetCharacter(InIndex))
+		{
+			FVector CharacterLocation = Character->GetTransform(/** Reset */false).GetLocation();
+
+			if (const UText3DLayoutExtensionBase* LayoutExtension = InTextComponent->GetLayoutExtension())
+			{
+				CharacterLocation *= LayoutExtension->GetTextScale();
+			}
+
+			return InTextComponent->GetComponentTransform().TransformPositionNoScale(CharacterLocation);
+		}
+	}
+
+	return FVector::ZeroVector;
+}
+
 void FAvaTextVisualizer::StoreTextMetrics(const UText3DComponent* InText3DComp)
 {
 	Bounds = GetComponentBounds(InText3DComp);
@@ -941,6 +1043,12 @@ void FAvaTextVisualizer::DrawVisualizationNotEditing(const UActorComponent* InCo
 			}
 		}
 	}
+
+	AvaTextComponent->ForEachCharacter([this, AvaTextComponent, InView, InPDI](UText3DCharacterBase* InCharacter, uint16 InIndex, uint16)
+	{
+		const FVector CharacterLocation = GetKerningHandleLocation(AvaTextComponent, InIndex);
+		DrawCharacterKerningHandle(AvaTextComponent, InView, InPDI, CharacterLocation, InIndex, Inactive);
+	});
 }
 
 void FAvaTextVisualizer::DrawVisualizationEditing(const UActorComponent* InComponent, const FSceneView* InView, 
@@ -959,31 +1067,51 @@ void FAvaTextVisualizer::DrawVisualizationEditing(const UActorComponent* InCompo
 
 	DrawMaxTextSizeVisualization(AvaTextComponent, InView, InPDI);
 
-	DrawMaxTextWidthButton(AvaTextComponent, InView, InPDI, InOutIconIndex, AvaTextComponent->HasMaxWidth() ? Enabled : Disabled);
-	++InOutIconIndex;
-
-	DrawMaxTextHeightButton(AvaTextComponent, InView, InPDI, InOutIconIndex, AvaTextComponent->HasMaxHeight() ? Enabled : Disabled);
-	++InOutIconIndex;
-
-	DrawScaleProportionallyButton(AvaTextComponent, InView, InPDI, InOutIconIndex, AvaTextComponent->ScalesProportionally() ? Enabled : Disabled);
-	++InOutIconIndex;
+	if (bEditingWidth)
+	{
+		DrawMaxTextWidthButton(AvaTextComponent, InView, InPDI, InOutIconIndex, AvaTextComponent->HasMaxWidth() ? Enabled : Disabled);
+		DrawMaxTextWidthHandle(AvaTextComponent, InView, InPDI, Inactive);
+		++InOutIconIndex;
+	}
+	
+	if (bEditingHeight)
+	{
+		DrawMaxTextHeightButton(AvaTextComponent, InView, InPDI, InOutIconIndex, AvaTextComponent->HasMaxHeight() ? Enabled : Disabled);
+		DrawMaxTextHeightHandle(AvaTextComponent, InView, InPDI, Inactive);
+		++InOutIconIndex;
+	}
 
-	DrawMaxTextWidthHandle(AvaTextComponent, InView, InPDI, Inactive);
-	DrawMaxTextHeightHandle(AvaTextComponent, InView, InPDI, Inactive);
+	if (bEditingWidth || bEditingHeight)
+	{
+		DrawScaleProportionallyButton(AvaTextComponent, InView, InPDI, InOutIconIndex, AvaTextComponent->ScalesProportionally() ? Enabled : Disabled);
+		++InOutIconIndex;
+	}
 
-	if (const UText3DDefaultMaterialExtension* MaterialExtension = TextComponent->GetCastedMaterialExtension<UText3DDefaultMaterialExtension>())
+	if (bEditingGradientOffset || bEditingGradientRotation_EndHandle || bEditingGradientRotation_StartHandle)
 	{
-		if (MaterialExtension->GetStyle() == EText3DMaterialStyle::Gradient)
+		if (const UText3DDefaultMaterialExtension* MaterialExtension = TextComponent->GetCastedMaterialExtension<UText3DDefaultMaterialExtension>())
 		{
-			DrawEditGradientButton(AvaTextComponent, InView, InPDI, InOutIconIndex, bShowGradientControls ? Enabled : Disabled);
-			++InOutIconIndex;
-
-			if (bShowGradientControls)
+			if (MaterialExtension->GetStyle() == EText3DMaterialStyle::Gradient)
 			{
-				DrawGradientHandles(AvaTextComponent, InView, InPDI);
+				DrawEditGradientButton(AvaTextComponent, InView, InPDI, InOutIconIndex, bShowGradientControls ? Enabled : Disabled);
+				++InOutIconIndex;
+
+				if (bShowGradientControls)
+				{
+					DrawGradientHandles(AvaTextComponent, InView, InPDI);
+				}
 			}
 		}
 	}
+
+	if (EditingKerningIndex != INDEX_NONE)
+	{
+		AvaTextComponent->ForEachCharacter([this, AvaTextComponent, InView, InPDI](UText3DCharacterBase* InCharacter, uint16 InIndex, uint16)
+		{
+			const FVector CharacterLocation = GetKerningHandleLocation(AvaTextComponent, InIndex);
+			DrawCharacterKerningHandle(AvaTextComponent, InView, InPDI, CharacterLocation, InIndex, InIndex == EditingKerningIndex ? Active : Inactive);
+		});
+	}
 }
 
 void FAvaTextVisualizer::StoreInitialValues()
@@ -1001,12 +1129,20 @@ void FAvaTextVisualizer::StoreInitialValues()
 	InitialMaxHeight = TextComponent->GetMaxHeight();
 	bInitialScaleProportionally = TextComponent->ScalesProportionally();
 
-	if (UText3DDefaultMaterialExtension* MaterialExtension = TextComponent->GetCastedMaterialExtension<UText3DDefaultMaterialExtension>())
+	if (const UText3DDefaultMaterialExtension* MaterialExtension = TextComponent->GetCastedMaterialExtension<UText3DDefaultMaterialExtension>())
 	{
 		InitialGradientRotation = MaterialExtension->GetGradientRotation();
 		InitialGradientOffset = MaterialExtension->GetGradientOffset();
 		InitialGradientSmoothness = MaterialExtension->GetGradientSmoothness();
 	}
+
+	if (EditingKerningIndex != INDEX_NONE)
+	{
+		if (const UText3DCharacterBase* Character = TextComponent->GetCharacter(EditingKerningIndex))
+		{
+			InitialCharacterKerning = Character->GetCharacterKerning().Get(0.f);
+		}
+	}
 }
 
 void FAvaTextVisualizer::DrawMaxTextSizeVisualization(const UText3DComponent* InTextComponent, const FSceneView* InView, 
@@ -1377,4 +1513,30 @@ void FAvaTextVisualizer::DrawGradientLineEndHandle(const UAvaText3DComponent* In
 	InPDI->SetHitProxy(nullptr);
 }
 
+void FAvaTextVisualizer::DrawCharacterKerningHandle(const UAvaText3DComponent* InTextComponent, const FSceneView* InView, FPrimitiveDrawInterface* InPDI,
+	const FVector& InLocation, uint16 InCharacterIndex, const FLinearColor& InColor) const
+{
+	static constexpr float BaseSize = 1.f;
+
+	if (!InTextComponent)
+	{
+		return;
+	}
+
+	const UTexture2D* const UVSprite = IAvalancheComponentVisualizersModule::Get().GetSettings()->GetVisualizerSprite(UE::AvaShapes::SizeSprite);
+
+	if (!UVSprite || !UVSprite->GetResource())
+	{
+		return;
+	}
+
+	const FVector IconLocation = InLocation;
+	const float IconSize = BaseSize * GetIconSizeScale(InView, IconLocation);
+
+	InPDI->SetHitProxy(new HAvaTextCharacterKerningHandleProxy(InTextComponent, InCharacterIndex));
+	InPDI->DrawSprite(IconLocation, IconSize, IconSize, UVSprite->GetResource(), InColor,
+		SDPG_Foreground, 0, 0, 0, 0, SE_BLEND_Opaque);
+	InPDI->SetHitProxy(nullptr);
+}
+
 #undef LOCTEXT_NAMESPACE
diff --git a/Engine/Plugins/Experimental/Avalanche/Source/AvalancheTextEditor/Private/Visualizer/AvaTextVisualizer.h b/Engine/Plugins/Experimental/Avalanche/Source/AvalancheTextEditor/Private/Visualizer/AvaTextVisualizer.h
--- a/Engine/Plugins/Experimental/Avalanche/Source/AvalancheTextEditor/Private/Visualizer/AvaTextVisualizer.h
+++ b/Engine/Plugins/Experimental/Avalanche/Source/AvalancheTextEditor/Private/Visualizer/AvaTextVisualizer.h
@@ -98,6 +98,18 @@ struct HAvaTextGradientSmoothnessHandleProxy : HAvaHitProxy
 	{}
 };
 
+struct HAvaTextCharacterKerningHandleProxy : HAvaHitProxy
+{
+	DECLARE_HIT_PROXY();
+
+	HAvaTextCharacterKerningHandleProxy(const UActorComponent* InComponent, uint16 InIndex)
+		: HAvaHitProxy(InComponent)
+		, Index(InIndex)
+	{}
+
+	uint16 Index;
+};
+
 class FAvaTextVisualizer : public FAvaVisualizerBase
 {
 public:
@@ -112,7 +124,6 @@ public:
 	virtual TMap<UObject*, TArray<FProperty*>> GatherEditableProperties(UObject* InObject) const override;
 	virtual bool VisProxyHandleClick(FEditorViewportClient* InViewportClient, HComponentVisProxy* VisProxy,
 		const FViewportClick& Click) override;
-	FVector GetGradientSmoothnessHandleLocation(const UAvaText3DComponent* InAvaText3DComponent) const;
 	virtual bool GetWidgetLocation(const FEditorViewportClient* ViewportClient, FVector& OutLocation) const override;
 	virtual bool GetWidgetMode(const FEditorViewportClient* ViewportClient, UE::Widget::EWidgetMode& Mode) const override;
 	virtual bool GetWidgetAxisList(const FEditorViewportClient* ViewportClient, UE::Widget::EWidgetMode WidgetMode,
@@ -133,6 +144,7 @@ protected:
 	FProperty* MaxWidthProperty;
 	FProperty* MaxHeightProperty;
 	FProperty* ScaleProportionallyProperty;
+	FProperty* CharacterKerningProperty;
 
 	TWeakObjectPtr<UAvaText3DComponent> TextComponent;
 	mutable FBox Bounds;
@@ -151,9 +163,12 @@ protected:
 	bool bEditingGradientRotation_EndHandle;
 	bool bEditingGradientOffset;
 	bool bEditingGradientSmoothness;
+	int32 EditingKerningIndex = INDEX_NONE;
+	
 	float InitialGradientRotation;
 	float InitialGradientOffset;
 	float InitialGradientSmoothness;
+	float InitialCharacterKerning;
 
 	FVector GradientEditBeginLocation_StartHandle;
 	FVector GradientEditBeginLocation_EndHandle;
@@ -195,6 +210,8 @@ protected:
 	FVector GetGradientEndHandleLocation(const UAvaText3DComponent* InTextComponent) const;
 	FVector GetGradientCenterHandleLocation(const UAvaText3DComponent* InTextComponent) const;
 	FVector GetGradientStartHandleLocation(const UAvaText3DComponent* InTextComponent) const;
+	FVector GetGradientSmoothnessHandleLocation(const UAvaText3DComponent* InAvaText3DComponent) const;
+	FVector GetKerningHandleLocation(const UAvaText3DComponent* InTextComponent, int32 InIndex) const;
 
 	void DrawMaxTextWidthButton(const UText3DComponent* InTextComponent, const FSceneView* InView, FPrimitiveDrawInterface* InPDI,
 		int32 InIconIndex, const FLinearColor& InColor) const;
@@ -219,4 +236,7 @@ protected:
 
 	void DrawGradientLineEndHandle(const UAvaText3DComponent* InTextComponent, const FSceneView* InView,
 		FPrimitiveDrawInterface* InPDI, const FVector& InLocation, const FLinearColor& InColor) const;
+	
+	void DrawCharacterKerningHandle(const UAvaText3DComponent* InTextComponent, const FSceneView* InView,
+		FPrimitiveDrawInterface* InPDI, const FVector& InLocation, uint16 InCharacterIndex, const FLinearColor& InColor) const;
 };
diff --git a/Engine/Plugins/Experimental/Text3D/Source/Text3D/Private/Characters/Text3DCharacterBase.cpp b/Engine/Plugins/Experimental/Text3D/Source/Text3D/Private/Characters/Text3DCharacterBase.cpp
new file mode 100644
--- /dev/null
+++ b/Engine/Plugins/Experimental/Text3D/Source/Text3D/Private/Characters/Text3DCharacterBase.cpp
@@ -0,0 +1,87 @@
+// Copyright Epic Games, Inc. All Rights Reserved.
+
+#include "Characters/Text3DCharacterBase.h"
+
+#include "Text3DComponent.h"
+
+FTransform& UText3DCharacterBase::GetTransform(bool bInReset)
+{
+	if (bInReset)
+	{
+		Transform = FTransform(RelativeRotation, RelativeLocation, RelativeScale);
+	}
+
+	return Transform;
+}
+
+void UText3DCharacterBase::SetCharacterMetadata(const FText3DCharacterMetadata& InMetadata)
+{
+	Metadata = InMetadata;
+}
+
+void UText3DCharacterBase::SetRelativeLocation(const FVector& InLocation)
+{
+	if (RelativeLocation.Equals(InLocation))
+	{
+		return;
+	}
+	
+	RelativeLocation = InLocation;
+	OnCharacterDataChanged();
+}
+
+void UText3DCharacterBase::SetRelativeRotation(const FRotator& InRotation)
+{
+	if (RelativeRotation.Equals(InRotation))
+	{
+		return;
+	}
+	
+	RelativeRotation = InRotation;
+	OnCharacterDataChanged();
+}
+
+void UText3DCharacterBase::SetRelativeScale(const FVector& InScale)
+{
+	if (RelativeScale.Equals(InScale))
+	{
+		return;
+	}
+	
+	RelativeScale = InScale;
+	OnCharacterDataChanged();
+}
+
+#if WITH_EDITOR
+void UText3DCharacterBase::PostEditUndo()
+{
+	Super::PostEditUndo();
+
+	OnCharacterDataChanged();
+}
+
+void UText3DCharacterBase::PostEditChangeProperty(FPropertyChangedEvent& InEvent)
+{
+	Super::PostEditChangeProperty(InEvent);
+
+	static const TSet<FName> PropertyNames =
+	{
+		GET_MEMBER_NAME_CHECKED(UText3DCharacterBase, RelativeLocation),
+		GET_MEMBER_NAME_CHECKED(UText3DCharacterBase, RelativeRotation),
+		GET_MEMBER_NAME_CHECKED(UText3DCharacterBase, RelativeScale)
+	};
+
+	if (PropertyNames.Contains(InEvent.GetMemberPropertyName()))
+	{
+		OnCharacterDataChanged();
+	}
+}
+#endif
+
+void UText3DCharacterBase::OnCharacterDataChanged() const
+{
+	if (UText3DComponent* Component = GetTypedOuter<UText3DComponent>())
+	{
+		Component->RequestUpdate(EText3DRendererFlags::Layout);
+	}
+}
diff --git a/Engine/Plugins/Experimental/Text3D/Source/Text3D/Private/Characters/Text3DDefaultCharacter.cpp b/Engine/Plugins/Experimental/Text3D/Source/Text3D/Private/Characters/Text3DDefaultCharacter.cpp
new file mode 100644
--- /dev/null
+++ b/Engine/Plugins/Experimental/Text3D/Source/Text3D/Private/Characters/Text3DDefaultCharacter.cpp
@@ -0,0 +1,50 @@
+// Copyright Epic Games, Inc. All Rights Reserved.
+
+#include "Characters/Text3DDefaultCharacter.h"
+
+FName UText3DDefaultCharacter::GetKerningPropertyName()
+{
+	return GET_MEMBER_NAME_CHECKED(UText3DDefaultCharacter, Kerning);
+}
+
+void UText3DDefaultCharacter::SetKerningOverride(bool bInKerningOverride)
+{
+	if (bKerningOverride == bInKerningOverride)
+	{
+		return;
+	}
+
+	bKerningOverride = bInKerningOverride;
+	OnCharacterDataChanged();
+}
+
+void UText3DDefaultCharacter::SetKerning(float InKerning)
+{
+	if (FMath::IsNearlyEqual(Kerning, InKerning))
+	{
+		return;
+	}
+
+	Kerning = InKerning;
+	OnCharacterDataChanged();
+}
+
+#if WITH_EDITOR
+void UText3DDefaultCharacter::PostEditChangeProperty(FPropertyChangedEvent& InEvent)
+{
+	Super::PostEditChangeProperty(InEvent);
+
+	const FName MemberName = InEvent.GetMemberPropertyName();
+
+	if (MemberName == GET_MEMBER_NAME_CHECKED(UText3DDefaultCharacter, Kerning)
+		|| MemberName == GET_MEMBER_NAME_CHECKED(UText3DDefaultCharacter, bKerningOverride))
+	{
+		OnCharacterDataChanged();
+	}
+}
+#endif
+
+TOptional<float> UText3DDefaultCharacter::GetCharacterKerning() const
+{
+	return bKerningOverride ? Kerning : TOptional<float>();
+}
diff --git a/Engine/Plugins/Experimental/Text3D/Source/Text3D/Private/Extensions/Text3DDefaultCharacterExtension.cpp b/Engine/Plugins/Experimental/Text3D/Source/Text3D/Private/Extensions/Text3DDefaultCharacterExtension.cpp
new file mode 100644
--- /dev/null
+++ b/Engine/Plugins/Experimental/Text3D/Source/Text3D/Private/Extensions/Text3DDefaultCharacterExtension.cpp
@@ -0,0 +1,99 @@
+// Copyright Epic Games, Inc. All Rights Reserved.
+
+#include "Extensions/Text3DDefaultCharacterExtension.h"
+
+#include "Characters/Text3DDefaultCharacter.h"
+#include "Logs/Text3DLogs.h"
+#include "Misc/EnumerateRange.h"
+#include "Text3DComponent.h"
+#include "Utilities/Text3DUtilities.h"
+
+uint16 UText3DDefaultCharacterExtension::GetCharacterCount() const
+{
+	return TextCharacters.Num();
+}
+
+UText3DCharacterBase* UText3DDefaultCharacterExtension::GetCharacter(uint16 InIndex) const
+{
+	if (TextCharacters.IsValidIndex(InIndex))
+	{
+		return TextCharacters[InIndex];
+	}
+
+	return nullptr;
+}
+
+TConstArrayView<UText3DCharacterBase*> UText3DDefaultCharacterExtension::GetCharacters() const
+{
+	return TextCharacters;
+}
+
+EText3DExtensionResult UText3DDefaultCharacterExtension::PreRendererUpdate(EText3DRendererFlags InFlag)
+{
+	if (InFlag != EText3DRendererFlags::Geometry)
+	{
+		return EText3DExtensionResult::Active;
+	}
+
+	UText3DComponent* Text3DComponent = GetText3DComponent();
+	const FTypefaceEntry* TypeFaceEntry = Text3DComponent->GetTypeFaceEntry();
+
+	if (!TypeFaceEntry)
+	{
+		UE_LOG(LogText3D, Warning, TEXT("Invalid typeface entry cannot proceed"))
+		return EText3DExtensionResult::Failed;
+	}
+
+	const FString FormattedTextStr = Text3DComponent->GetFormattedText().ToString();
+	const TArray<FText3DCharacterMetadata> TextCharacterMetadata = UE::Text3D::Utilities::GetCharacterMetadata(*TypeFaceEntry, FormattedTextStr);
+
+	AllocateTextCharacters(TextCharacterMetadata.Num());
+
+	for (TConstEnumerateRef<TObjectPtr<UText3DCharacterBase>> TextCharacter : EnumerateRange(TextCharacters))
+	{
+		const FText3DCharacterMetadata& Metadata = TextCharacterMetadata[TextCharacter.GetIndex()];
+#if WITH_EDITORONLY_DATA
+		(*TextCharacter)->SetCharacter(FString::Printf(TEXT("%i - %c"), Metadata.TextIndex, FormattedTextStr[Metadata.TextIndex]));
+#endif
+		(*TextCharacter)->SetCharacterMetadata(Metadata);
+	}
+
+	return EText3DExtensionResult::Finished;
+}
+
+EText3DExtensionResult UText3DDefaultCharacterExtension::PostRendererUpdate(EText3DRendererFlags InFlag)
+{
+	return EText3DExtensionResult::Active;
+}
+
+void UText3DDefaultCharacterExtension::AllocateTextCharacters(uint16 InCharacterCount)
+{
+	TextCharacters.Reserve(InCharacterCount);
+
+	const int32 RemainingCharacterCount = TextCharacters.Num() - InCharacterCount;
+
+	if (RemainingCharacterCount > 0)
+	{
+		for (int32 CharacterIndex = InCharacterCount; CharacterIndex < TextCharacters.Num(); CharacterIndex++)
+		{
+			TextCharactersPool.Add(TextCharacters[CharacterIndex]);
+		}
+	}
+	else if (RemainingCharacterCount < 0)
+	{
+		for (int32 CharacterIndex = 0; CharacterIndex < FMath::Abs(RemainingCharacterCount); ++CharacterIndex)
+		{
+			if (!TextCharactersPool.IsEmpty())
+			{
+				TextCharacters.Add(TextCharactersPool.Pop());
+			}
+			else
+			{
+				const FName ObjectName = MakeUniqueObjectName(this, UText3DDefaultCharacter::StaticClass(), FName(TEXT("Char")));
+				TextCharacters.Add(NewObject<UText3DCharacterBase>(this, UText3DDefaultCharacter::StaticClass(), ObjectName, RF_NoFlags));
+			}
+		}
+	}
+
+	TextCharacters.SetNum(InCharacterCount);
+}
diff --git a/Engine/Plugins/Experimental/Text3D/Source/Text3D/Private/Extensions/Text3DDefaultGeometryExtension.cpp b/Engine/Plugins/Experimental/Text3D/Source/Text3D/Private/Extensions/Text3DDefaultGeometryExtension.cpp
--- a/Engine/Plugins/Experimental/Text3D/Source/Text3D/Private/Extensions/Text3DDefaultGeometryExtension.cpp
+++ b/Engine/Plugins/Experimental/Text3D/Source/Text3D/Private/Extensions/Text3DDefaultGeometryExtension.cpp
@@ -2,8 +2,11 @@
 
 #include "Extensions/Text3DDefaultGeometryExtension.h"
 
-#include "Text3DComponent.h"
+#include "Characters/Text3DCharacterBase.h"
+#include "Engine/StaticMesh.h"
+#include "MeshDescription.h"
 #include "Subsystems/Text3DEngineSubsystem.h"
+#include "Text3DComponent.h"
 
 void UText3DDefaultGeometryExtension::SetUseOutline(const bool bValue)
 {
@@ -27,6 +30,34 @@ void UText3DDefaultGeometryExtension::SetOutline(const float Value)
 	OnGeometryOptionsChanged();
 }
 
+void UText3DDefaultGeometryExtension::SetPivotHAlignment(EText3DHorizontalTextAlignment InPivot)
+{
+	if (PivotHAlignment == InPivot)
+	{
+		return;
+	}
+
+	PivotHAlignment = InPivot;
+	OnGeometryOptionsChanged();
+}
+
+void UText3DDefaultGeometryExtension::SetPivotVAlignment(EText3DVerticalTextAlignment InPivot)
+{
+	if (PivotVAlignment == InPivot)
+	{
+		return;
+	}
+
+	// Only one supported for now
+	if (PivotVAlignment != EText3DVerticalTextAlignment::Bottom)
+	{
+		return;
+	}
+
+	PivotVAlignment = InPivot;
+	OnGeometryOptionsChanged();
+}
+
 void UText3DDefaultGeometryExtension::SetExtrude(const float Value)
 {
 	const float NewValue = FMath::Max(0.0f, Value);
@@ -93,6 +124,8 @@ void UText3DDefaultGeometryExtension::PostEditChangeProperty(FPropertyChangedEve
 		GET_MEMBER_NAME_CHECKED(UText3DDefaultGeometryExtension, BevelSegments),
 		GET_MEMBER_NAME_CHECKED(UText3DDefaultGeometryExtension, bUseOutline),
 		GET_MEMBER_NAME_CHECKED(UText3DDefaultGeometryExtension, Outline),
+		GET_MEMBER_NAME_CHECKED(UText3DDefaultGeometryExtension, PivotHAlignment),
+		GET_MEMBER_NAME_CHECKED(UText3DDefaultGeometryExtension, PivotVAlignment),
 	};
 
 	if (PropertyNames.Contains(InEvent.GetMemberPropertyName()))
@@ -102,10 +135,23 @@ void UText3DDefaultGeometryExtension::PostEditChangeProperty(FPropertyChangedEve
 }
 #endif
 
-UStaticMesh* UText3DDefaultGeometryExtension::GetGlyphMesh(uint32 InGlyphIndex, const FFreeTypeFace* InFontFace) const
+UStaticMesh* UText3DDefaultGeometryExtension::GetGlyphMesh(UText3DCharacterBase* InCharacter) const
+{
+	const uint16 GlyphIndex = InCharacter->GetCharacterMetadata().GlyphIndex;
+	const FText3DCachedMesh* GlyphMesh = CachedFontData->GetGlyphMesh(GlyphIndex, GlyphMeshParameters);
+	checkf(GlyphMesh != nullptr, TEXT("Invalid glyph mesh returned"));
+	InCharacter->GetCharacterMetadata().MeshBounds = GlyphMesh->MeshBounds;
+	return GlyphMesh->Mesh;
+}
+
+EText3DHorizontalTextAlignment UText3DDefaultGeometryExtension::GetGlyphHAlignment() const
+{
+	return PivotHAlignment;
+}
+
+EText3DVerticalTextAlignment UText3DDefaultGeometryExtension::GetGlyphVAlignment() const
 {
-	check(!!CachedFontData)
-	return CachedFontData->GetGlyphMesh(InGlyphIndex, GlyphMeshParameters, InFontFace);
+	return PivotVAlignment;
 }
 
 EText3DExtensionResult UText3DDefaultGeometryExtension::PreRendererUpdate(EText3DRendererFlags InFlag)
@@ -118,8 +164,25 @@ EText3DExtensionResult UText3DDefaultGeometryExtension::PreRendererUpdate(EText3
 	const UText3DComponent* TextComponent = GetText3DComponent();
 	UText3DEngineSubsystem* Subsystem = UText3DEngineSubsystem::Get();
 	const uint32 TypefaceIndex = TextComponent->GetTypeFaceIndex();
+
 	CachedFontData = &Subsystem->GetCachedFontData(TextComponent->GetFont(), TypefaceIndex);
-	GlyphMeshParameters = {GetExtrude(), GetBevel(), GetBevelType(), GetBevelSegments(), GetUseOutline(), GetOutline(), TypefaceIndex};
+
+	if (!CachedFontData)
+	{
+		return EText3DExtensionResult::Failed;
+	}
+
+	GlyphMeshParameters =
+	{
+		.Extrude = GetExtrude(),
+		.Bevel = GetBevel(),
+		.BevelType = GetBevelType(),
+		.BevelSegments = GetBevelSegments(),
+		.bOutline = GetUseOutline(),
+		.OutlineExpand = GetOutline(),
+		.TypefaceIndex = TypefaceIndex,
+		.PivotOffset = GetPivotOffset(),
+	};
 
 	// Keep to clean after
 	return EText3DExtensionResult::Active;
@@ -175,3 +238,18 @@ float UText3DDefaultGeometryExtension::GetMaxBevel() const
 {
 	return Extrude / 2.0f;
 }
+
+FVector UText3DDefaultGeometryExtension::GetPivotOffset() const
+{
+	switch (PivotHAlignment)
+	{
+	case EText3DHorizontalTextAlignment::Left:
+		return FVector::ZeroVector;
+	case EText3DHorizontalTextAlignment::Center:
+		return FVector(0.0f, -0.5f, 0.0f);
+	case EText3DHorizontalTextAlignment::Right:
+		return FVector(0.0f, -1.f, 0.0f);
+	}
+
+	return FVector::ZeroVector;
+}
diff --git a/Engine/Plugins/Experimental/Text3D/Source/Text3D/Private/Extensions/Text3DDefaultLayoutExtension.cpp b/Engine/Plugins/Experimental/Text3D/Source/Text3D/Private/Extensions/Text3DDefaultLayoutExtension.cpp
--- a/Engine/Plugins/Experimental/Text3D/Source/Text3D/Private/Extensions/Text3DDefaultLayoutExtension.cpp
+++ b/Engine/Plugins/Experimental/Text3D/Source/Text3D/Private/Extensions/Text3DDefaultLayoutExtension.cpp
@@ -2,16 +2,20 @@
 
 #include "Extensions/Text3DDefaultLayoutExtension.h"
 
+#include "Algo/Accumulate.h"
+#include "Characters/Text3DCharacterBase.h"
 #include "Engine/Font.h"
+#include "Extensions/Text3DGeometryExtensionBase.h"
+#include "Extensions/Text3DTransformExtensionBase.h"
 #include "Fonts/SlateTextShaper.h"
 #include "Framework/Text/PlainTextLayoutMarshaller.h"
-#include "GeometryBuilders/Text3DGlyphMeshBuilder.h"
-#include "GeometryBuilders/Text3DGlyph.h"
 #include "LayoutBuilders/Text3DLayout.h"
 #include "LayoutBuilders/Text3DLayoutShaper.h"
 #include "LayoutBuilders/Text3DShapedGlyphText.h"
 #include "Logs/Text3DLogs.h"
+#include "Misc/EnumerateRange.h"
 #include "Styling/StyleDefaults.h"
+#include "Styling/SlateTypes.h"
 #include "Subsystems/Text3DEngineSubsystem.h"
 #include "Text3DComponent.h"
 
@@ -163,14 +167,9 @@ void UText3DDefaultLayoutExtension::SetScaleProportionally(const bool Value)
 	OnLayoutOptionsChanged();
 }
 
-UText3DDefaultLayoutExtension::UText3DDefaultLayoutExtension()
-{
-	ShapedText = MakeShared<FText3DShapedGlyphText>();
-}
-
 float UText3DDefaultLayoutExtension::GetTextHeight() const
 {
-	return ShapedText->Lines.Num() * ShapedText->LineHeight + (ShapedText->Lines.Num() - 1) * LineSpacing;
+	return ShapedText->Lines.Num() * ShapedText->FontHeight + (ShapedText->Lines.Num() - 1) * LineSpacing;
 }
 
 FVector UText3DDefaultLayoutExtension::GetTextScale() const
@@ -180,100 +179,149 @@ FVector UText3DDefaultLayoutExtension::GetTextScale() const
 
 EText3DExtensionResult UText3DDefaultLayoutExtension::PreRendererUpdate(EText3DRendererFlags InFlag)
 {
-	if (InFlag != EText3DRendererFlags::Geometry && InFlag != EText3DRendererFlags::Layout)
+	if (InFlag != EText3DRendererFlags::Layout)
 	{
 		return EText3DExtensionResult::Active;
 	}
 
-	const UText3DComponent* Text3DComponent = GetText3DComponent();
+	UText3DComponent* Text3DComponent = GetText3DComponent();
 	
-	if (!EnumHasAnyFlags(Text3DComponent->GetUpdateFlags(), EText3DRendererFlags::Geometry))
+	if (!ShapedText.IsValid())
 	{
-		ShapedText->Kerning = Kerning;
-		ShapedText->WordSpacing = WordSpacing;
-		ShapedText->MaxWidth = MaxWidth;
-		ShapedText->bWrap = bUseMaxWidth && MaxWidthBehavior == EText3DMaxWidthHandling::WrapAndScale;
+		ShapedText = MakeShared<FText3DShapedGlyphText>();
 	}
-	else
+
+	if (!TextLayout.IsValid())
 	{
-		UText3DEngineSubsystem* Subsystem = UText3DEngineSubsystem::Get();
+		TextLayout = MakeShared<FText3DLayout>();
+	}
 
-		const uint32 TypefaceIndex = Text3DComponent->GetTypeFaceIndex();
-		FCachedFontData& CachedFontData = Subsystem->GetCachedFontData(Text3DComponent->GetFont(), TypefaceIndex);
-		const FT_Face Face = CachedFontData.GetFreeTypeFace(TypefaceIndex);
+	if (!TextLayoutMarshaller.IsValid())
+	{
+		TextLayoutMarshaller = FPlainTextLayoutMarshaller::Create();
+	}
 
-		if (!Face)
-		{
-			UE_LOG(LogText3D, Error, TEXT("Failed to load font face data '%s' in Text3D layout extension"), *CachedFontData.GetFontName());
-			return EText3DExtensionResult::Failed;
-		}
+	if (!FontStyle.IsValid())
+	{
+		FontStyle = MakeShared<FTextBlockStyle>();
+	}
 
-		ShapedText->Reset();
-		ShapedText->LineHeight = Face->size->metrics.height * FontInverseScale;
-		ShapedText->FontAscender = Face->size->metrics.ascender * FontInverseScale;
-		ShapedText->FontDescender = Face->size->metrics.descender * FontInverseScale;
-		ShapedText->Kerning = Kerning;
-		ShapedText->WordSpacing = WordSpacing;
-		ShapedText->MaxWidth = MaxWidth;
-		ShapedText->bWrap = bUseMaxWidth && MaxWidthBehavior == EText3DMaxWidthHandling::WrapAndScale;
-
-		constexpr int32 AdjustedFontSize = 48; // Magic number that makes font scale consistent with previous implementation
-		FSlateFontInfo FontInfo(Text3DComponent->GetFont(), AdjustedFontSize);
+	if (EnumHasAnyFlags(Text3DComponent->GetUpdateFlags(), EText3DRendererFlags::Geometry))
+	{
+		FSlateFontInfo FontInfo(Text3DComponent->GetFont(), UE::Text3D::Metrics::FontSize);
 		FontInfo.CompositeFont = FStyleDefaults::GetFontInfo().CompositeFont;
 		FontInfo.TypefaceFontName = Text3DComponent->GetTypeface();
+		FontStyle->SetFont(FontInfo);
+	}
 
-		FTextBlockStyle Style;
-		Style.SetFont(FontInfo);
+	ShapedText->Reset();
+	ShapedText->FontHeight = FontStyle->Font.Size;
+	FText3DLayoutShaper::Get()->ShapeBidirectionalText(*FontStyle, Text3DComponent->GetFormattedText().ToString(), TextLayout, TextLayoutMarshaller, ShapedText->Lines);
 
-		if (!TextLayout.IsValid())
+	const int32 ApproximateGlyphNum = Algo::TransformAccumulate(ShapedText->Lines,
+		[](const FShapedGlyphLine& InLine)
 		{
-			TextLayout = MakeShared<FText3DLayout>(Style);
-		}
+			return InLine.GlyphsToRender.Num();
+		},
+		0);
 
-		if (!TextLayoutMarshaller.IsValid())
+	ShapedText->GlyphIndexToFontFace.Reserve(ApproximateGlyphNum);
+	for (int32 LineIndex = 0; LineIndex < ShapedText->Lines.Num(); ++LineIndex)
+	{
+		for (const FShapedGlyphEntry& GlyphEntry : ShapedText->Lines[LineIndex].GlyphsToRender)
 		{
-			TextLayoutMarshaller = FPlainTextLayoutMarshaller::Create();
-		}
+			if (!GlyphEntry.FontFaceData.IsValid())
+			{
+				// Add as nullptr if not already in map
+				if (!ShapedText->GlyphIndexToFontFace.Contains(GlyphEntry.GlyphIndex))
+				{
+					ShapedText->GlyphIndexToFontFace.Add(GlyphEntry.GlyphIndex, nullptr);
+				}
 
-		FText3DLayoutShaper::Get()->ShapeBidirectionalText(Style, Text3DComponent->GetFormattedText().ToString(), TextLayout, TextLayoutMarshaller, ShapedText->Lines);
+				continue;
+			}
+
+			const TSharedPtr<FFreeTypeFace> FontFacePtr = GlyphEntry.FontFaceData->FontFace.Pin();
 
-		const int32 ApproximateGlyphNum = Algo::TransformAccumulate(ShapedText->Lines,
-			[](const FShapedGlyphLine& InLine)
+			if (FontFacePtr.IsValid())
 			{
-				return InLine.GlyphsToRender.Num();
-			},
-			0);
+				ShapedText->GlyphIndexToFontFace.FindOrAdd(GlyphEntry.GlyphIndex, FontFacePtr.Get());
+			}
+		}
+	}
+
+	ShapedText->CharacterGlyphMap.Reset();
+	ShapedText->Kernings.SetNum(Text3DComponent->GetCharacterCount());
+	ShapedText->GlyphSizes.SetNum(Text3DComponent->GetCharacterCount());
+	Text3DComponent->ForEachCharacter([this](UText3DCharacterBase* InCharacter, uint16 InIndex, uint16 InCount)
+	{
+		ShapedText->Kernings[InIndex] = InCharacter->GetCharacterKerning().Get(Kerning);
+		ShapedText->GlyphSizes[InIndex] = InCharacter->GetCharacterMetadata().MeshBounds.GetSize();
+		ShapedText->CharacterGlyphMap.Add(InCharacter->GetCharacterMetadata().LayoutIndex, InIndex);
+	});
+
+	UText3DGeometryExtensionBase* GeometryExtension = Text3DComponent->GetGeometryExtension();
+	ShapedText->PivotHAlignment = GeometryExtension->GetGlyphHAlignment();
+	ShapedText->PivotVAlignment = GeometryExtension->GetGlyphVAlignment();
+	ShapedText->Kerning = Kerning;
+	ShapedText->WordSpacing = WordSpacing;
+	ShapedText->MaxWidth = MaxWidth;
+	ShapedText->bWrap = bUseMaxWidth && MaxWidthBehavior == EText3DMaxWidthHandling::WrapAndScale;
+	ShapedText->CalculateWidth();
+	CalculateTextScale();
 
-		ShapedText->GlyphIndexToFontFace.Reserve(ApproximateGlyphNum);
+	const TConstArrayView<UText3DTransformExtensionBase*> TransformExtensions = Text3DComponent->GetTransformExtensions();
 
-		for (int32 LineIndex = 0; LineIndex < ShapedText->Lines.Num(); ++LineIndex)
+	int32 LayoutIndex = 0;
+	for (TConstEnumerateRef<FShapedGlyphLine> ShapedLine : EnumerateRange(ShapedText->Lines))
+	{
+		FVector Location = GetLineLocation(ShapedLine.GetIndex());
+
+		for (TConstEnumerateRef<FShapedGlyphEntry> ShapedGlyph : EnumerateRange(ShapedLine->GlyphsToRender))
 		{
-			for (const FShapedGlyphEntry& GlyphEntry : ShapedText->Lines[LineIndex].GlyphsToRender)
+			if (const int32* CharacterIndex = ShapedText->CharacterGlyphMap.Find(LayoutIndex))
 			{
-				if (!GlyphEntry.FontFaceData.IsValid())
+				UText3DCharacterBase* CharacterData = Text3DComponent->GetCharacter(*CharacterIndex);
+
+				constexpr bool bReset = true;
+				FTransform& CharacterTransform = CharacterData->GetTransform(bReset);
+
+				const FVector GlyphSize = ShapedText->GlyphSizes[*CharacterIndex];
+				FVector Shift = FVector::ZeroVector;
+				switch (ShapedText->PivotHAlignment)
 				{
-					// Add as nullptr if not already in map
-					if (!ShapedText->GlyphIndexToFontFace.Contains(GlyphEntry.GlyphIndex))
-					{
-						ShapedText->GlyphIndexToFontFace.Add(GlyphEntry.GlyphIndex, nullptr);
-					}
+				case EText3DHorizontalTextAlignment::Left:
+					Shift = FVector::ZeroVector;
+					break;
+				case EText3DHorizontalTextAlignment::Center:
+					Shift = GlyphSize / 2;
+					break;
+				case EText3DHorizontalTextAlignment::Right:
+					Shift = GlyphSize;
+					break;
+				}
 
-					continue;
+				if (ShapedGlyph.GetIndex() != 0)
+				{
+					Shift.Y += ShapedText->Kernings[*CharacterIndex];
 				}
 
-				const TSharedPtr<FFreeTypeFace> FontFacePtr = GlyphEntry.FontFaceData->FontFace.Pin();
+				CharacterTransform.Accumulate(FTransform(Location + FVector(0, Shift.Y, 0)));
 
-				if (FontFacePtr.IsValid())
+				for (UText3DTransformExtensionBase* TransformExtension : TransformExtensions)
 				{
-					ShapedText->GlyphIndexToFontFace.FindOrAdd(GlyphEntry.GlyphIndex, FontFacePtr.Get());
+					if (TransformExtension && TransformExtension->GetTargetRange().IsInRange(*CharacterIndex))
+					{
+						TransformExtension->ApplyTransform(*CharacterIndex, ShapedText->CharacterGlyphMap.Num(), CharacterTransform);
+					}
 				}
 			}
+
+			Location.Y += ShapedLine->GlyphAdvances[ShapedGlyph.GetIndex()];
+			LayoutIndex++;
 		}
 	}
 
-	ShapedText->CalculateWidth();
-	CalculateTextScale();
-
 	return EText3DExtensionResult::Finished;
 }
 
@@ -325,7 +373,7 @@ FVector UText3DDefaultLayoutExtension::GetLineLocation(int32 LineIndex) const
 	float VerticalOffset = 0.0f;
 	if (LineIndex < 0 || LineIndex >= ShapedText->Lines.Num())
 	{
-		return FVector();
+		return FVector::ZeroVector;
 	}
 
 	const FShapedGlyphLine& ShapedLine = ShapedText->Lines[LineIndex];
@@ -343,7 +391,7 @@ FVector UText3DDefaultLayoutExtension::GetLineLocation(int32 LineIndex) const
 	if (VerticalAlignment != EText3DVerticalTextAlignment::FirstLine)
 	{
 		// First align it to Top
-		VerticalOffset -= ShapedText->FontAscender;
+		VerticalOffset -= ShapedText->FontHeight;
 
 		if (VerticalAlignment == EText3DVerticalTextAlignment::Center)
 		{
@@ -351,20 +399,15 @@ FVector UText3DDefaultLayoutExtension::GetLineLocation(int32 LineIndex) const
 		}
 		else if (VerticalAlignment == EText3DVerticalTextAlignment::Bottom)
 		{
-			VerticalOffset += TotalHeight + ShapedText->FontDescender;
+			VerticalOffset += TotalHeight;
 		}
 	}
 
-	VerticalOffset -= LineIndex * (ShapedText->LineHeight + LineSpacing);
+	VerticalOffset -= LineIndex * (ShapedText->FontHeight + LineSpacing);
 
 	return FVector(0.0f, HorizontalOffset, VerticalOffset);
 }
 
-TSharedPtr<FText3DShapedGlyphText> UText3DDefaultLayoutExtension::GetShapedText() const
-{
-	return ShapedText;
-}
-
 #if WITH_EDITOR
 void UText3DDefaultLayoutExtension::PostEditChangeProperty(FPropertyChangedEvent& InEvent)
 {
@@ -397,14 +440,6 @@ void UText3DDefaultLayoutExtension::OnLayoutOptionsChanged()
 	MaxHeight = FMath::Max(1.0f, MaxHeight);
 
 	EText3DRendererFlags Flags = EText3DRendererFlags::Material;
-	if (MaxWidthBehavior == EText3DMaxWidthHandling::WrapAndScale)
-	{
-		EnumAddFlags(Flags, EText3DRendererFlags::Geometry);
-	}
-	else
-	{
-		EnumAddFlags(Flags, EText3DRendererFlags::Layout);
-	}
-
+	EnumAddFlags(Flags, EText3DRendererFlags::Layout);
 	RequestUpdate(Flags);
 }
diff --git a/Engine/Plugins/Experimental/Text3D/Source/Text3D/Private/Extensions/Text3DDefaultMaterialExtension.cpp b/Engine/Plugins/Experimental/Text3D/Source/Text3D/Private/Extensions/Text3DDefaultMaterialExtension.cpp
--- a/Engine/Plugins/Experimental/Text3D/Source/Text3D/Private/Extensions/Text3DDefaultMaterialExtension.cpp
+++ b/Engine/Plugins/Experimental/Text3D/Source/Text3D/Private/Extensions/Text3DDefaultMaterialExtension.cpp
@@ -317,8 +317,8 @@ EText3DExtensionResult UText3DDefaultMaterialExtension::PreRendererUpdate(EText3
 		const UText3DComponent* Text3DComponent = GetText3DComponent();
 
 		const FBox LocalBounds = Text3DComponent->GetBounds();
-		const FVector LocalBoundsLowerLeftCorner = LocalBounds.Min;
-		const FVector LocalBoundsExtent = LocalBounds.GetExtent();
+		const FVector LocalBoundsLowerLeftCorner = Text3DComponent->GetRelativeRotation().RotateVector(LocalBounds.Min);
+		const FVector LocalBoundsExtent = Text3DComponent->GetRelativeRotation().RotateVector(LocalBounds.GetExtent());
 
 		const UText3DLayoutExtensionBase* LayoutExtension = Text3DComponent->GetLayoutExtension();
 
diff --git a/Engine/Plugins/Experimental/Text3D/Source/Text3D/Private/GeometryBuilders/Text3DGlyphData.cpp b/Engine/Plugins/Experimental/Text3D/Source/Text3D/Private/GeometryBuilders/Text3DGlyphData.cpp
--- a/Engine/Plugins/Experimental/Text3D/Source/Text3D/Private/GeometryBuilders/Text3DGlyphData.cpp
+++ b/Engine/Plugins/Experimental/Text3D/Source/Text3D/Private/GeometryBuilders/Text3DGlyphData.cpp
@@ -42,13 +42,13 @@ void FText3DGlyphData::SetCurrentGroup(const EText3DGroupType Type, const float
 	Group.FirstVertex = MeshDescription.Vertices().Num();
 	Group.FirstTriangle = MeshDescription.Triangles().Num();
 
-	GroupExpand = GroupExpandIn / FontInverseScale;
+	GroupExpand = GroupExpandIn / UE::Text3D::Metrics::FontSizeInverse;
 }
 
 void FText3DGlyphData::PrepareSegment(const float PlannedExtrudeIn, const float PlannedExpandIn, const FVector2D NormalStartIn, const FVector2D NormalEndIn)
 {
 	PlannedExtrude = PlannedExtrudeIn;
-	PlannedExpand = PlannedExpandIn / FontInverseScale;
+	PlannedExpand = PlannedExpandIn / UE::Text3D::Metrics::FontSizeInverse;
 
 	NormalStart = NormalStartIn;
 	NormalEnd = NormalEndIn;
@@ -197,7 +197,7 @@ void FText3DGlyphData::FillEdge(const FText3DGlyphPartPtr& Edge, const bool bSki
 
 FVector FText3DGlyphData::GetVector(const FVector2D Position, const float Height) const
 {
-	return (FVector(0.f, Position.X, Position.Y) * FontInverseScale + FVector(Height, 0.0f, 0.0f));
+	return (FVector(0.f, Position.X, Position.Y) * UE::Text3D::Metrics::FontSizeInverse + FVector(Height, 0.0f, 0.0f));
 }
 
 void FText3DGlyphData::MakeTriangleFanAlongNormal(const FText3DGlyphPartConstPtr& Cap, const FText3DGlyphPartPtr& Normal, const bool bNormalIsCapNext, const bool bSkipLastTriangle)
diff --git a/Engine/Plugins/Experimental/Text3D/Source/Text3D/Private/GeometryBuilders/Text3DGlyphMeshBuilder.cpp b/Engine/Plugins/Experimental/Text3D/Source/Text3D/Private/GeometryBuilders/Text3DGlyphMeshBuilder.cpp
--- a/Engine/Plugins/Experimental/Text3D/Source/Text3D/Private/GeometryBuilders/Text3DGlyphMeshBuilder.cpp
+++ b/Engine/Plugins/Experimental/Text3D/Source/Text3D/Private/GeometryBuilders/Text3DGlyphMeshBuilder.cpp
@@ -95,11 +95,49 @@ void FText3DGlyphMeshBuilder::MirrorGroups(const float Extrude)
 	}
 }
 
-void FText3DGlyphMeshBuilder::BuildMesh(UStaticMesh* StaticMesh, class UMaterial* DefaultMaterial)
+void FText3DGlyphMeshBuilder::MovePivot(const FVector& InNewPivot) const
+{
+	const TVertexAttributesRef<FVector3f> VertexPositions = Glyph->GetStaticMeshAttributes().GetVertexPositions();
+
+	// Remove empty space around glyph first
+	float MinX = TNumericLimits<float>::Max();
+	float MaxX = TNumericLimits<float>::Lowest();
+	float MinY = TNumericLimits<float>::Max();
+	float MaxY = TNumericLimits<float>::Lowest();
+	for (const FVertexID VertexID : Glyph->GetMeshDescription().Vertices().GetElementIDs())
+	{
+		const FVector3f Position = VertexPositions[VertexID];
+		MinX = FMath::Min(MinX, Position.X);
+		MaxX = FMath::Max(MaxX, Position.X);
+		MinY = FMath::Min(MinY, Position.Y);
+		MaxY = FMath::Max(MaxY, Position.Y);
+	}
+
+	// Apply offset with adjusted values
+	const float CenterX = (MinX + MaxX) * 0.5f;
+	FVector3f OffsetVector(-CenterX, -MinY, 0.0f);
+
+	const FVector3f PivotOffset = FVector3f(InNewPivot) * FVector3f(MaxX - MinX, MaxY - MinY, 0);
+	OffsetVector += PivotOffset;
+
+	for (const FVertexID VertexID : Glyph->GetMeshDescription().Vertices().GetElementIDs())
+	{
+		FVector3f CurrentPosition = VertexPositions[VertexID];
+		CurrentPosition += OffsetVector;
+		VertexPositions[VertexID] = CurrentPosition;
+	}
+}
+
+void FText3DGlyphMeshBuilder::BuildMesh(UStaticMesh* StaticMesh, UMaterial* DefaultMaterial)
 {
 	Glyph->Build(StaticMesh, DefaultMaterial);
 }
 
+FBox FText3DGlyphMeshBuilder::GetMeshBounds() const
+{
+	return Glyph->GetMeshDescription().ComputeBoundingBox();
+}
+
 void FText3DGlyphMeshBuilder::CreateFrontMesh(const TText3DGlyphContourNodeShared& Root, const bool bOutline, const float& OutlineExpand)
 {
 	int32 VertexCount = 0;
diff --git a/Engine/Plugins/Experimental/Text3D/Source/Text3D/Private/GeometryBuilders/Text3DGlyphMeshBuilder.h b/Engine/Plugins/Experimental/Text3D/Source/Text3D/Private/GeometryBuilders/Text3DGlyphMeshBuilder.h
--- a/Engine/Plugins/Experimental/Text3D/Source/Text3D/Private/GeometryBuilders/Text3DGlyphMeshBuilder.h
+++ b/Engine/Plugins/Experimental/Text3D/Source/Text3D/Private/GeometryBuilders/Text3DGlyphMeshBuilder.h
@@ -31,8 +31,15 @@ public:
 	void CreateMeshes(const TText3DGlyphContourNodeShared& Root, const float Extrude, const float Bevel, const EText3DBevelType Type, const int32 BevelSegments, const bool bOutline, const float OutlineExpand = 0.5f);
 	void SetFrontAndBevelTextureCoordinates(const float Bevel);
 	void MirrorGroups(const float Extrude);
+
+	/** Move the mesh pivot based on pivot ratio 0-1 */
+	void MovePivot(const FVector& InNewPivot) const;
+
 	void BuildMesh(UStaticMesh* StaticMesh, UMaterial* DefaultMaterial);
 
+	/** Get actual glyph bounds once built */
+	FBox GetMeshBounds() const;
+
 private:
 	TSharedPtr<FText3DGlyph> Glyph;
 	TSharedRef<FText3DGlyphData> Data;
diff --git a/Engine/Plugins/Experimental/Text3D/Source/Text3D/Private/LayoutBuilders/Text3DLayout.h b/Engine/Plugins/Experimental/Text3D/Source/Text3D/Private/LayoutBuilders/Text3DLayout.h
--- a/Engine/Plugins/Experimental/Text3D/Source/Text3D/Private/LayoutBuilders/Text3DLayout.h
+++ b/Engine/Plugins/Experimental/Text3D/Source/Text3D/Private/LayoutBuilders/Text3DLayout.h
@@ -12,6 +12,11 @@ public:
 	/** Optionally provide a custom TextBlock style. */
 	explicit FText3DLayout(const FTextBlockStyle& InStyle = FTextBlockStyle::GetDefault());
 
+	const FTextBlockStyle& GetTextStyle() const
+	{
+		return TextStyle;
+	}
+
 protected:
 	/** Parameters relevant to text layout. */
 	FTextBlockStyle TextStyle;
diff --git a/Engine/Plugins/Experimental/Text3D/Source/Text3D/Private/LayoutBuilders/Text3DShapedGlyphLine.cpp b/Engine/Plugins/Experimental/Text3D/Source/Text3D/Private/LayoutBuilders/Text3DShapedGlyphLine.cpp
--- a/Engine/Plugins/Experimental/Text3D/Source/Text3D/Private/LayoutBuilders/Text3DShapedGlyphLine.cpp
+++ b/Engine/Plugins/Experimental/Text3D/Source/Text3D/Private/LayoutBuilders/Text3DShapedGlyphLine.cpp
@@ -2,18 +2,23 @@
 
 #include "LayoutBuilders/Text3DShapedGlyphLine.h"
 
-float FShapedGlyphLine::GetAdvance(const int32 InIndex, const float InKerning, const float InWordSpacing) const
+float FShapedGlyphLine::GetWidthAdvance(const int32 InIndex, const float InKerning, const float InWordSpacing) const
 {
 	check(InIndex >= 0 && InIndex < GlyphsToRender.Num());
 
 	const FShapedGlyphEntry& Glyph = GlyphsToRender[InIndex];
+
 	float Advance = Glyph.XOffset + Glyph.XAdvance;
 
-	if (InIndex < GlyphsToRender.Num() - 1)
+	// Skip first character kerning
+	if (InIndex != 0)
 	{
 		// @note: as per FSlateElementBatcher::BuildShapedTextSequence, per Glyph Kerning isn't used
 		Advance += InKerning;
+	}
 
+	if (InIndex < GlyphsToRender.Num() - 1)
+	{
 		if (!Glyph.bIsVisible)
 		{
 			Advance += InWordSpacing;
@@ -22,12 +27,3 @@ float FShapedGlyphLine::GetAdvance(const int32 InIndex, const float InKerning, c
 
 	return Advance;
 }
-
-void FShapedGlyphLine::CalculateWidth(const float InKerning, const float InWordSpacing)
-{
-	Width = 0.0f;
-	for (int32 Index = 0; Index < GlyphsToRender.Num(); Index++)
-	{
-		Width += GetAdvance(Index, InKerning, InWordSpacing);
-	}
-}
diff --git a/Engine/Plugins/Experimental/Text3D/Source/Text3D/Private/LayoutBuilders/Text3DShapedGlyphLine.h b/Engine/Plugins/Experimental/Text3D/Source/Text3D/Private/LayoutBuilders/Text3DShapedGlyphLine.h
--- a/Engine/Plugins/Experimental/Text3D/Source/Text3D/Private/LayoutBuilders/Text3DShapedGlyphLine.h
+++ b/Engine/Plugins/Experimental/Text3D/Source/Text3D/Private/LayoutBuilders/Text3DShapedGlyphLine.h
@@ -4,19 +4,20 @@
 
 #include "Containers/Array.h"
 #include "Fonts/FontCache.h"
+#include "Text3DTypes.h"
 
 /** Contains text line metric with sufficient information to fetch and transform each character. */
 struct FShapedGlyphLine
 {
 	/** Get's the offset from the previous character, accounting for kerning and word spacing. */
-	float GetAdvance(const int32 InIndex, const float InKerning, const float InWordSpacing) const;
-
-	/** Calculates the total width of this line. */
-	void CalculateWidth(const float InKerning, const float InWordSpacing);
+	float GetWidthAdvance(const int32 InIndex, const float InKerning, const float InWordSpacing) const;
 
 	/** The corresponding shaped glyph for each character in this line of text. */
 	TArray<FShapedGlyphEntry> GlyphsToRender;
 
-	/** Stored result of CalculateWidth. */
+	/** Stored result of line width */
 	float Width = 0.0f;
+
+	/** Glyph advance on this line */
+	TArray<float> GlyphAdvances;
 };
\ No newline at end of file
diff --git a/Engine/Plugins/Experimental/Text3D/Source/Text3D/Private/LayoutBuilders/Text3DShapedGlyphText.cpp b/Engine/Plugins/Experimental/Text3D/Source/Text3D/Private/LayoutBuilders/Text3DShapedGlyphText.cpp
--- a/Engine/Plugins/Experimental/Text3D/Source/Text3D/Private/LayoutBuilders/Text3DShapedGlyphText.cpp
+++ b/Engine/Plugins/Experimental/Text3D/Source/Text3D/Private/LayoutBuilders/Text3DShapedGlyphText.cpp
@@ -11,14 +11,17 @@ FText3DShapedGlyphText::FText3DShapedGlyphText()
 
 void FText3DShapedGlyphText::Reset()
 {
-	LineHeight = 0.0f;
-	FontAscender = 0.0f;
-	FontDescender = 0.0f;
-	Kerning = 0.0f;
+	FontHeight = 0.0f;
+	PivotHAlignment = EText3DHorizontalTextAlignment::Left;
+	PivotVAlignment = EText3DVerticalTextAlignment::Bottom;
+	Kernings.Reset();
+	GlyphSizes.Reset();
+	Kerning = 0.f;
 	WordSpacing = 0.0f;
 	bWrap = false;
 	Lines.Reset();
-	GlyphIndexToFontFace.Empty();
+	GlyphIndexToFontFace.Reset();
+	CharacterGlyphMap.Reset();
 }
 
 void FText3DShapedGlyphText::CalculateWidth()
@@ -26,27 +29,31 @@ void FText3DShapedGlyphText::CalculateWidth()
 	TArray<FShapedGlyphLine> NewLines;
 	NewLines.Reserve(Lines.Num());
 
-	int32 TextGlyphIdx = 0;
-	for (const FShapedGlyphLine& GlyphLine : Lines)
+	int32 LayoutIndex = 0;
+	for (const TConstEnumerateRef<FShapedGlyphLine> GlyphLine : EnumerateRange(Lines))
 	{
 		FShapedGlyphLine* CurrentLine = &NewLines.AddDefaulted_GetRef();
-		TArray<FShapedGlyphEntry> CurrentWord;
 
+		TArray<float> CurrentAdvances;
+		TArray<FShapedGlyphEntry> CurrentWord;
 		float LineWidth = 0.0f;
 		float CurrentWordLength = 0.0f;
 
-		for (TConstEnumerateRef<FShapedGlyphEntry> CurrentGlyph : EnumerateRange(GlyphLine.GlyphsToRender))
+		for (TConstEnumerateRef<FShapedGlyphEntry> CurrentGlyph : EnumerateRange(GlyphLine->GlyphsToRender))
 		{
 			// Trim for proper positioning
 			if (!CurrentGlyph->bIsVisible && CurrentGlyph.GetIndex() == 0)
 			{
+				LayoutIndex++;
 				continue;
 			}
 
-			const bool bWordBreak = !CurrentGlyph->bIsVisible || CurrentGlyph.GetIndex() == GlyphLine.GlyphsToRender.Num() - 1;
-			const float GlyphAdv = GlyphLine.GetAdvance(CurrentGlyph.GetIndex(), Kerning, WordSpacing);
+			const int32* CharacterIndex = CharacterGlyphMap.Find(LayoutIndex);
+			const float CurrentKerning = !!CharacterIndex ? Kernings[*CharacterIndex] : Kerning;
+			const float GlyphAdv = GlyphLine->GetWidthAdvance(CurrentGlyph.GetIndex(), CurrentKerning, WordSpacing);
 
 			// If we're at the end the line or at whitespace
+			const bool bWordBreak = !CurrentGlyph->bIsVisible || CurrentGlyph.GetIndex() == GlyphLine->GlyphsToRender.Num() - 1;
 			if (bWrap                               // when we're wrapping
 				&& bWordBreak                       // and at a word break
 				&& LineWidth > MaxWidth             // and the current line is longer than the max
@@ -60,44 +67,23 @@ void FText3DShapedGlyphText::CalculateWidth()
 			CurrentWord.Add(*CurrentGlyph);
 			LineWidth += GlyphAdv;
 			CurrentWordLength += GlyphAdv;
+			CurrentAdvances.Add(GlyphAdv);
 
 			if (bWordBreak)
 			{
 				CurrentLine->GlyphsToRender.Append(CurrentWord);
+				CurrentLine->GlyphAdvances.Append(CurrentAdvances);
 				CurrentWordLength = 0.0f;
-				CurrentWord.Empty();
+				CurrentWord.Reset();
+				CurrentAdvances.Reset();
 			}
 
-			TextGlyphIdx++;
+			LayoutIndex++;
 		}
 
 		CurrentLine->Width = LineWidth;
 	}
-
-	TextGlyphIdx = 0;
-	for (FShapedGlyphLine& NewLine : NewLines)
-	{
-		if (NewLine.GlyphsToRender.IsEmpty())
-		{
-			continue;
-		}
-
-		const FShapedGlyphEntry& FirstGlyph = NewLine.GlyphsToRender[0];
-		if (!FirstGlyph.bIsVisible)
-		{
-			NewLine.Width -= NewLine.GetAdvance(0, Kerning, WordSpacing);
-			NewLine.GlyphsToRender.RemoveAt(0);
-		}
-
-		const int32 LastIndex = NewLine.GlyphsToRender.Num() - 1;
-		TextGlyphIdx += LastIndex;
-		const FShapedGlyphEntry& LastGlyph = NewLine.GlyphsToRender[LastIndex];
-		if (!LastGlyph.bIsVisible)
-		{
-			NewLine.Width -= NewLine.GetAdvance(LastIndex, Kerning, WordSpacing);
-			NewLine.GlyphsToRender.RemoveAt(LastIndex);
-		}
-	}
 	
 	Lines = NewLines;
 }
+
diff --git a/Engine/Plugins/Experimental/Text3D/Source/Text3D/Private/LayoutBuilders/Text3DShapedGlyphText.h b/Engine/Plugins/Experimental/Text3D/Source/Text3D/Private/LayoutBuilders/Text3DShapedGlyphText.h
--- a/Engine/Plugins/Experimental/Text3D/Source/Text3D/Private/LayoutBuilders/Text3DShapedGlyphText.h
+++ b/Engine/Plugins/Experimental/Text3D/Source/Text3D/Private/LayoutBuilders/Text3DShapedGlyphText.h
@@ -16,13 +16,39 @@ struct FText3DShapedGlyphText
 	void Reset();
 	void CalculateWidth();
 
-	float LineHeight;
-	float FontAscender;
-	float FontDescender;
+	/** Maps character layout index to character index */
+	TMap<int32, int32> CharacterGlyphMap;
+
+	/** Font height expected */
+	float FontHeight;
+
+	/** Pivot Horizontal alignment */
+	EText3DHorizontalTextAlignment PivotHAlignment;
+
+	/** Pivot Vertical alignment */
+	EText3DVerticalTextAlignment PivotVAlignment;
+
+	/** Individual character kerning */
+	TArray<float> Kernings;
+
+	/** Individual character size */
+	TArray<FVector> GlyphSizes;
+
+	/** General kerning for invisible character */
 	float Kerning;
+
+	/** Spacing between words */
 	float WordSpacing;
+
+	/** Maximum width allowed */
 	float MaxWidth;
+
+	/** Whether to wrap words */
 	bool bWrap;
+
+	/** Lines calculated based on current text */
 	TArray<FShapedGlyphLine> Lines;
+
+	/** Glyph index to font face */
 	TMap<uint32, const FFreeTypeFace*> GlyphIndexToFontFace;
 };
diff --git a/Engine/Plugins/Experimental/Text3D/Source/Text3D/Private/Renderers/StaticMeshes/Text3DStaticMeshesRenderer.cpp b/Engine/Plugins/Experimental/Text3D/Source/Text3D/Private/Renderers/StaticMeshes/Text3DStaticMeshesRenderer.cpp
--- a/Engine/Plugins/Experimental/Text3D/Source/Text3D/Private/Renderers/StaticMeshes/Text3DStaticMeshesRenderer.cpp
+++ b/Engine/Plugins/Experimental/Text3D/Source/Text3D/Private/Renderers/StaticMeshes/Text3DStaticMeshesRenderer.cpp
@@ -4,6 +4,7 @@
 
 #include "Algo/Accumulate.h"
 #include "Algo/Count.h"
+#include "Characters/Text3DCharacterBase.h"
 #include "Components/SceneComponent.h"
 #include "Components/StaticMeshComponent.h"
 #include "Engine/StaticMesh.h"
@@ -18,7 +19,6 @@
 #include "Subsystems/Text3DEngineSubsystem.h"
 #include "Text3DComponent.h"
 #include "Text3DInternalTypes.h"
-#include "Extensions/Text3DTransformExtensionBase.h"
 
 void UText3DStaticMeshesRenderer::OnCreate()
 {
@@ -36,61 +36,15 @@ void UText3DStaticMeshesRenderer::OnUpdate(EText3DRendererFlags InFlags)
 
 	if (EnumHasAnyFlags(InFlags, EText3DRendererFlags::Geometry))
 	{
-		const UText3DLayoutExtensionBase* LayoutExtension = TextComponent->GetLayoutExtension();
-
-		TextRoot->SetRelativeScale3D(LayoutExtension->GetTextScale());
-
-		const TSharedPtr<FText3DShapedGlyphText> ShapedText = LayoutExtension->GetShapedText();
-
-		// Pre-allocate, avoid new'ing!
-		AllocateGlyphs(Algo::TransformAccumulate(ShapedText->Lines, [](const FShapedGlyphLine& ShapedLine)
-		{
-			return Algo::CountIf(ShapedLine.GlyphsToRender, [](const FShapedGlyphEntry& ShapedGlyph)
-			{
-				return ShapedGlyph.bIsVisible;
-			});
-		},
-		0));
+		AllocateCharacterComponents(TextComponent->GetCharacterCount());
 
 		const UText3DGeometryExtensionBase* GeometryExtension = TextComponent->GetGeometryExtension();
 
-		int32 GlyphIndex = 0;
-		for (TConstEnumerateRef<FShapedGlyphLine> ShapedLine : EnumerateRange(ShapedText->Lines))
+		TextComponent->ForEachCharacter([this, GeometryExtension](UText3DCharacterBase* InCharacter, uint16 InIndex, uint16 InCount)
 		{
-			FVector LineLocation = LayoutExtension->GetLineLocation(ShapedLine.GetIndex());
-
-			for (TConstEnumerateRef<FShapedGlyphEntry> ShapedGlyph : EnumerateRange(ShapedLine->GlyphsToRender))
-			{
-				FVector GlyphLocation = LineLocation;
-				LineLocation.Y += ShapedLine->GetAdvance(ShapedGlyph.GetIndex(), ShapedText->Kerning, ShapedText->WordSpacing);
-
-				if (!ShapedGlyph->bIsVisible)
-				{
-					continue;
-				}
-
-				// Count even when mesh is nullptr (allocation still creates components to avoid mesh building in allocation step)
-				const int32 GlyphId = GlyphIndex++;
-
-				const FFreeTypeFace* FontFace = ShapedText->GlyphIndexToFontFace[ShapedGlyph->GlyphIndex];
-				UStaticMesh* CachedMesh = GeometryExtension->GetGlyphMesh(ShapedGlyph->GlyphIndex, FontFace);
-				if (!CachedMesh || FMath::IsNearlyZero(CachedMesh->GetBounds().SphereRadius))
-				{
-					continue;
-				}
-
-				if (CharacterMeshes.IsValidIndex(GlyphId))
-				{
-					UStaticMeshComponent* StaticMeshComponent = CharacterMeshes[GlyphId];
-					StaticMeshComponent->SetStaticMesh(CachedMesh);
-				}
-				else
-				{
-					// @note: This shouldn't occur, but it does under unknown circumstances (UE-164789) so it should be handled
-					UE_LOG(LogText3D, Error, TEXT("CharacterMesh not found at index %d"), GlyphId);
-				}
-			}
-		}
+			UStaticMesh* CachedMesh = GeometryExtension->GetGlyphMesh(InCharacter);
+			CharacterMeshes[InIndex]->SetStaticMesh(CachedMesh);
+		});
 
 		RefreshBounds();
 	}
@@ -101,50 +55,13 @@ void UText3DStaticMeshesRenderer::OnUpdate(EText3DRendererFlags InFlags)
 
 		TextRoot->SetRelativeScale3D(LayoutExtension->GetTextScale());
 
-		const TSharedPtr<FText3DShapedGlyphText> ShapedText = LayoutExtension->GetShapedText();
-
-		const TConstArrayView<UText3DTransformExtensionBase*> TransformExtensions = TextComponent->GetTransformExtensions();
-
-		const int32 GlyphCount = CharacterMeshes.Num();
-		int32 GlyphIndex = 0;
-		for (TConstEnumerateRef<FShapedGlyphLine> ShapedLine : EnumerateRange(ShapedText->Lines))
+		TextComponent->ForEachCharacter([this](UText3DCharacterBase* InCharacter, uint16 InIndex, uint16 InCount)
 		{
-			FVector Location = LayoutExtension->GetLineLocation(ShapedLine.GetIndex());
-
-			for (TConstEnumerateRef<FShapedGlyphEntry> ShapedGlyph : EnumerateRange(ShapedLine->GlyphsToRender))
-			{
-				const FVector CharLocation = Location;
-				Location.Y += ShapedLine->GetAdvance(ShapedGlyph.GetIndex(), ShapedText->Kerning, ShapedText->WordSpacing);
-
-				if (!ShapedGlyph->bIsVisible)
-				{
-					continue;
-				}
+			constexpr bool bReset = false;
+			const FTransform& CharacterTransform = InCharacter->GetTransform(bReset);
+			CharacterMeshes[InIndex]->SetRelativeTransform(CharacterTransform);
+		});
 
-				if (USceneComponent* GlyphKerningComponent = GetGlyphKerningComponent(GlyphIndex))
-				{
-					GlyphKerningComponent->SetRelativeLocation(CharLocation);
-				}
-
-				if (UStaticMeshComponent* GlyphMeshComponent = GetGlyphMeshComponent(GlyphIndex))
-				{
-					FTransform GlyphTransform = FTransform::Identity;
-
-					for (UText3DTransformExtensionBase* TransformExtension : TransformExtensions)
-					{
-						if (TransformExtension && TransformExtension->GetTargetRange().IsInRange(GlyphIndex))
-						{
-							TransformExtension->ApplyTransform(GlyphIndex, GlyphCount, GlyphTransform);
-						}
-					}
-
-					GlyphMeshComponent->SetRelativeTransform(GlyphTransform);
-				}
-
-				GlyphIndex++;
-			}
-		}
-		
 		RefreshBounds();
 	}
 
@@ -152,7 +69,7 @@ void UText3DStaticMeshesRenderer::OnUpdate(EText3DRendererFlags InFlags)
 	{
 		using namespace UE::Text3D::Materials;
 
-		const UText3DMaterialExtensionBase* MaterialExtensionBase = TextComponent->GetMaterialExtension();
+		const UText3DMaterialExtensionBase* MaterialExtension = TextComponent->GetMaterialExtension();
 
 		for (UStaticMeshComponent* StaticMeshComponent : CharacterMeshes)
 		{
@@ -165,7 +82,7 @@ void UText3DStaticMeshesRenderer::OnUpdate(EText3DRendererFlags InFlags)
 					continue;
 				}
 
-				UMaterialInterface* Material = MaterialExtensionBase->GetMaterial(static_cast<EText3DGroupType>(GroupIndex));
+				UMaterialInterface* Material = MaterialExtension->GetMaterial(static_cast<EText3DGroupType>(GroupIndex));
 
 				if (Material != StaticMeshComponent->GetMaterial(MaterialIndex))
 				{
@@ -256,7 +173,7 @@ FBox UText3DStaticMeshesRenderer::OnCalculateBounds() const
 	return Box;
 }
 
-bool UText3DStaticMeshesRenderer::AllocateGlyphs(int32 Num)
+bool UText3DStaticMeshesRenderer::AllocateCharacterComponents(int32 Num)
 {
 	int32 DeltaNum = Num - CharacterMeshes.Num();
 	if (DeltaNum == 0)
diff --git a/Engine/Plugins/Experimental/Text3D/Source/Text3D/Private/Renderers/StaticMeshes/Text3DStaticMeshesRenderer.h b/Engine/Plugins/Experimental/Text3D/Source/Text3D/Private/Renderers/StaticMeshes/Text3DStaticMeshesRenderer.h
--- a/Engine/Plugins/Experimental/Text3D/Source/Text3D/Private/Renderers/StaticMeshes/Text3DStaticMeshesRenderer.h
+++ b/Engine/Plugins/Experimental/Text3D/Source/Text3D/Private/Renderers/StaticMeshes/Text3DStaticMeshesRenderer.h
@@ -59,7 +59,7 @@ protected:
 	//~ End UText3DRendererBase
 
 	/** Allocates, or shrinks existing components to match the input number. Returns false if nothing modified. */
-	bool AllocateGlyphs(int32 Num);
+	bool AllocateCharacterComponents(int32 Num);
 
 	/** Holds all character components */
 	UPROPERTY(Transient)
diff --git a/Engine/Plugins/Experimental/Text3D/Source/Text3D/Private/Renderers/Text3DRendererBase.cpp b/Engine/Plugins/Experimental/Text3D/Source/Text3D/Private/Renderers/Text3DRendererBase.cpp
--- a/Engine/Plugins/Experimental/Text3D/Source/Text3D/Private/Renderers/Text3DRendererBase.cpp
+++ b/Engine/Plugins/Experimental/Text3D/Source/Text3D/Private/Renderers/Text3DRendererBase.cpp
@@ -3,6 +3,7 @@
 #include "Renderers/Text3DRendererBase.h"
 
 #include "Extensions/Text3DGeometryExtensionBase.h"
+#include "Extensions/Text3DCharacterExtensionBase.h"
 #include "Extensions/Text3DLayoutExtensionBase.h"
 #include "Extensions/Text3DMaterialExtensionBase.h"
 #include "Extensions/Text3DRenderingExtensionBase.h"
@@ -52,8 +53,9 @@ void UText3DRendererBase::Update(EText3DRendererFlags InFlags)
 		return;
 	}
 
-	TArray<UText3DExtensionBase*, TInlineAllocator<5>> Extensions
+	TArray<UText3DExtensionBase*, TInlineAllocator<6>> Extensions
 	{
+		Text3DComponent->GetCharacterExtension(),
 		Text3DComponent->GetGeometryExtension(),
 		Text3DComponent->GetLayoutExtension(),
 		Text3DComponent->GetMaterialExtension(),
@@ -62,6 +64,17 @@ void UText3DRendererBase::Update(EText3DRendererFlags InFlags)
 
 	Extensions.Append(Text3DComponent->GetTransformExtensions());
 
+	Extensions.RemoveAll([](const UText3DExtensionBase* InExtension)
+	{
+		return !IsValid(InExtension);
+	});
+
+	// Sort in reverse order since we iterate in reverse order
+	Extensions.StableSort([](const UText3DExtensionBase& InExtensionA, const UText3DExtensionBase& InExtensionB)->bool
+	{
+		return InExtensionA.GetUpdatePriority() > InExtensionB.GetUpdatePriority();
+	});
+
 	uint8 Flag = 1 << 0;
 
 	while (Flag < static_cast<uint8>(EText3DRendererFlags::All))
diff --git a/Engine/Plugins/Experimental/Text3D/Source/Text3D/Private/Subsystems/Text3DEngineSubsystem.cpp b/Engine/Plugins/Experimental/Text3D/Source/Text3D/Private/Subsystems/Text3DEngineSubsystem.cpp
--- a/Engine/Plugins/Experimental/Text3D/Source/Text3D/Private/Subsystems/Text3DEngineSubsystem.cpp
+++ b/Engine/Plugins/Experimental/Text3D/Source/Text3D/Private/Subsystems/Text3DEngineSubsystem.cpp
@@ -250,8 +250,7 @@ void FCachedFontData::LoadFreeTypeFace(uint32 InTypefaceIndex)
 	if (const FT_Face* FreeTypeface = &TypefaceFontDataMap[InTypefaceIndex].GetTypeface())
 	{
 		TypefaceFontDataMap[InTypefaceIndex].SetTypefaceFontDataHash(HashCombine(0, GetTypeHash(Typeface.Font)));
-		FT_Set_Char_Size(*FreeTypeface, FontSize, FontSize, 96, 96);
-		FT_Set_Pixel_Sizes(*FreeTypeface, FontSize, FontSize);
+		FT_Set_Pixel_Sizes(*FreeTypeface, 0, UE::Text3D::Metrics::FontSize);
 	}
 }
 
@@ -289,7 +288,7 @@ TSharedPtr<int32> FCachedFontData::GetMeshesCacheCounter(const FGlyphMeshParamet
 	return nullptr;
 }
 
-UStaticMesh* FCachedFontData::GetGlyphMesh(uint32 InGlyphIndex, const FGlyphMeshParameters& InParameters, const FFreeTypeFace* InFontFaceData)
+FText3DCachedMesh* FCachedFontData::GetGlyphMesh(uint32 InGlyphIndex, const FGlyphMeshParameters& InParameters)
 {
 	const uint32 TypefaceIndex = InParameters.TypefaceIndex;
 	const uint32 HashParameters = GetTypeHash(InParameters);
@@ -301,10 +300,9 @@ UStaticMesh* FCachedFontData::GetGlyphMesh(uint32 InGlyphIndex, const FGlyphMesh
 
 	FCachedFontMeshes& CachedMeshes = TypefaceFontDataMap[TypefaceIndex].FindOrAddMeshes(HashParameters);
 
-	TObjectPtr<UStaticMesh>* CachedStaticMesh = CachedMeshes.Glyphs.Find(InGlyphIndex);
-	if (CachedStaticMesh)
+	if (FText3DCachedMesh* CachedMesh = CachedMeshes.Glyphs.Find(InGlyphIndex))
 	{
-		return *CachedStaticMesh;
+		return CachedMesh;
 	}
 
 	uint32 HashGroup = 0;
@@ -312,7 +310,7 @@ UStaticMesh* FCachedFontData::GetGlyphMesh(uint32 InGlyphIndex, const FGlyphMesh
 	HashGroup = HashCombine(HashGroup, GetTypeHash(InGlyphIndex));
 	const FString StaticMeshName = FString::Printf(TEXT("Text3D_Char_%u_%u"), HashGroup, HashParameters);
 
-	const TText3DGlyphContourNodeShared Root = GetGlyphContours(InGlyphIndex, TypefaceIndex, InFontFaceData);
+	const TText3DGlyphContourNodeShared Root = GetGlyphContours(InGlyphIndex, TypefaceIndex);
 	if (Root->Children.Num() == 0)
 	{
 		return nullptr;
@@ -321,17 +319,20 @@ UStaticMesh* FCachedFontData::GetGlyphMesh(uint32 InGlyphIndex, const FGlyphMesh
 	UText3DEngineSubsystem* Subsystem = UText3DEngineSubsystem::Get();
 	const FName StrokeShapeUniqueName = MakeUniqueObjectName(Subsystem, UStaticMesh::StaticClass(), FName(*StaticMeshName));
 
-	UStaticMesh* StaticMesh = NewObject<UStaticMesh>(Subsystem, StrokeShapeUniqueName);
-	CachedMeshes.Glyphs.Add(InGlyphIndex, StaticMesh);
+	FText3DCachedMesh& CachedMesh = CachedMeshes.Glyphs.FindOrAdd(InGlyphIndex);
+	CachedMesh.Mesh = NewObject<UStaticMesh>(Subsystem, StrokeShapeUniqueName);
 
 	FText3DGlyphMeshBuilder MeshCreator;
 	MeshCreator.CreateMeshes(Root, InParameters.Extrude, InParameters.Bevel, InParameters.BevelType, InParameters.BevelSegments, InParameters.bOutline, InParameters.OutlineExpand);
 	MeshCreator.SetFrontAndBevelTextureCoordinates(InParameters.Bevel);
 	MeshCreator.MirrorGroups(InParameters.Extrude);
+	MeshCreator.MovePivot(InParameters.PivotOffset);
 
-	MeshCreator.BuildMesh(StaticMesh, Subsystem->DefaultMaterial);
+	MeshCreator.BuildMesh(CachedMesh.Mesh, Subsystem->DefaultMaterial);
 
-	return StaticMesh;
+	CachedMesh.MeshBounds = MeshCreator.GetMeshBounds();
+
+	return &CachedMesh;
 }
 
 FT_Face FCachedFontData::GetFreeTypeFace(uint32 InTypefaceIndex)
@@ -354,36 +355,20 @@ FString FCachedFontData::GetFontName() const
 	return TEXT("");
 }
 
-TText3DGlyphContourNodeShared FCachedFontData::GetGlyphContours(uint32 InGlyphIndex, int32 InTypefaceEntryIndex, const FFreeTypeFace* InFontFaceData)
+TText3DGlyphContourNodeShared FCachedFontData::GetGlyphContours(uint32 InGlyphIndex, int32 InTypefaceEntryIndex)
 {
-	const bool bFontFaceValid = InFontFaceData && InFontFaceData->IsFaceValid();
-	if (!bFontFaceValid && !GetFreeTypeFace(InTypefaceEntryIndex))
+	FT_Face Face = GetFreeTypeFace(InTypefaceEntryIndex);
+	if (!Face)
 	{
 		return nullptr;
 	}
 
-	const FT_Face& GlyphFreeTypeFace = bFontFaceValid
-		? InFontFaceData->GetFace()
-		: GetFreeTypeFace(InTypefaceEntryIndex);
-	check(GlyphFreeTypeFace);
-
-	if (FT_Load_Glyph(GlyphFreeTypeFace, InGlyphIndex, FT_LOAD_NO_HINTING | FT_LOAD_NO_BITMAP))
+	if (FT_Load_Glyph(Face, InGlyphIndex, FT_LOAD_NO_HINTING | FT_LOAD_NO_BITMAP) != 0)
 	{
 		return nullptr;
 	}
 
-#if UE_BUILD_DEBUG
-	// If true, the symbol itself was probably bitmap-only, and this will have placeholder contours
-	bool bIsLastResort = InFontFaceData && InFontFaceData->GetFace()->family_name
-		? FString(InFontFaceData->GetFace()->family_name) == TEXT("LastResort")
-		: false;
-	if (bIsLastResort)
-	{
-		UE_LOG(LogText3D, Warning, TEXT("The following glyph index is bitmap-only: %i"), InGlyphIndex);
-	}
-#endif
-
-	FText3DGlyphLoader GlyphLoader(GlyphFreeTypeFace->glyph);
+	FText3DGlyphLoader GlyphLoader(Face->glyph);
 	TText3DGlyphContourNodeShared Root = GlyphLoader.GetContourList();
 
 	return Root;
diff --git a/Engine/Plugins/Experimental/Text3D/Source/Text3D/Private/Subsystems/Text3DEngineSubsystem.h b/Engine/Plugins/Experimental/Text3D/Source/Text3D/Private/Subsystems/Text3DEngineSubsystem.h
--- a/Engine/Plugins/Experimental/Text3D/Source/Text3D/Private/Subsystems/Text3DEngineSubsystem.h
+++ b/Engine/Plugins/Experimental/Text3D/Source/Text3D/Private/Subsystems/Text3DEngineSubsystem.h
@@ -21,7 +21,6 @@ struct FGlyphMeshParameters
 {
 	GENERATED_BODY()
 
-public:
 	UPROPERTY()
 	float Extrude = 5.0f;
 
@@ -42,34 +41,49 @@ public:
 
 	UPROPERTY()
 	uint32 TypefaceIndex = 0;
+
+	UPROPERTY()
+	FVector PivotOffset = FVector::ZeroVector;
+
+	friend uint32 GetTypeHash(const FGlyphMeshParameters& InElement)
+	{
+		uint32 HashParameters = 0;
+		HashParameters = HashCombine(HashParameters, GetTypeHash(InElement.Extrude));
+		HashParameters = HashCombine(HashParameters, GetTypeHash(InElement.Bevel));
+		HashParameters = HashCombine(HashParameters, GetTypeHash(InElement.BevelType));
+		HashParameters = HashCombine(HashParameters, GetTypeHash(InElement.BevelSegments));
+		HashParameters = HashCombine(HashParameters, GetTypeHash(InElement.bOutline));
+		HashParameters = HashCombine(HashParameters, GetTypeHash(InElement.OutlineExpand));
+		HashParameters = HashCombine(HashParameters, GetTypeHash(InElement.TypefaceIndex));
+		HashParameters = HashCombine(HashParameters, GetTypeHash(InElement.PivotOffset));
+		return HashParameters;
+	}
 };
 
-inline uint32 GetTypeHash(const FGlyphMeshParameters A)
+USTRUCT()
+struct FText3DCachedMesh
 {
-	uint32 HashParameters = 0;
-	HashParameters = HashCombine(HashParameters, GetTypeHash(A.Extrude));
-	HashParameters = HashCombine(HashParameters, GetTypeHash(A.Bevel));
-    HashParameters = HashCombine(HashParameters, GetTypeHash(A.BevelType));
-	HashParameters = HashCombine(HashParameters, GetTypeHash(A.BevelSegments));
-	HashParameters = HashCombine(HashParameters, GetTypeHash(A.bOutline));
-	HashParameters = HashCombine(HashParameters, GetTypeHash(A.OutlineExpand));
-	HashParameters = HashCombine(HashParameters, GetTypeHash(A.TypefaceIndex));
-	return HashParameters;
-}
+	GENERATED_BODY()
+
+	UPROPERTY()
+	TObjectPtr<UStaticMesh> Mesh;
+
+	UPROPERTY()
+	FBox MeshBounds;
+};
 
 USTRUCT()
 struct FCachedFontMeshes
 {
 	GENERATED_BODY()
 
-public:
 	FCachedFontMeshes();
 
 	int32 GetCacheCount() const;
 	TSharedPtr<int32> GetCacheCounter();
 
 	UPROPERTY()
-	TMap<uint32, TObjectPtr<UStaticMesh>> Glyphs;
+	TMap<uint32, FText3DCachedMesh> Glyphs;
 
 private:
 	TSharedPtr<int32> CacheCounter;
@@ -80,7 +94,6 @@ struct FTypefaceFontData
 {
 	GENERATED_BODY()
 
-public:
 	FTypefaceFontData();
 
 	void Reset()
@@ -143,7 +156,6 @@ struct FCachedFontData
 {
 	GENERATED_BODY()
 
-public:
 	FCachedFontData();
 	~FCachedFontData();
 
@@ -168,8 +180,8 @@ public:
 	TSharedPtr<int32> GetCacheCounter(int32 InTypefaceEntryIndex);
 	TSharedPtr<int32> GetMeshesCacheCounter(const FGlyphMeshParameters& InParameters);
 
-	UStaticMesh* GetGlyphMesh(uint32 InGlyphIndex, const FGlyphMeshParameters& InParameters, const FFreeTypeFace* InFontFaceData = nullptr);
-	TText3DGlyphContourNodeShared GetGlyphContours(uint32 InGlyphIndex, int32 InTypefaceEntryIndex, const FFreeTypeFace* InFontFaceData = nullptr);
+	FText3DCachedMesh* GetGlyphMesh(uint32 InGlyphIndex, const FGlyphMeshParameters& InParameters);
+	TText3DGlyphContourNodeShared GetGlyphContours(uint32 InGlyphIndex, int32 InTypefaceEntryIndex);
 
 	void PrintCache() const;
 
diff --git a/Engine/Plugins/Experimental/Text3D/Source/Text3D/Private/Text3DActor.cpp b/Engine/Plugins/Experimental/Text3D/Source/Text3D/Private/Text3DActor.cpp
--- a/Engine/Plugins/Experimental/Text3D/Source/Text3D/Private/Text3DActor.cpp
+++ b/Engine/Plugins/Experimental/Text3D/Source/Text3D/Private/Text3DActor.cpp
@@ -8,3 +8,10 @@ AText3DActor::AText3DActor()
 	Text3DComponent = CreateDefaultSubobject<UText3DComponent>(TEXT("Text3DComponent"));
 	RootComponent = Text3DComponent;
 }
+
+#if WITH_EDITOR
+FString AText3DActor::GetDefaultActorLabel() const
+{
+	return TEXT("Text3DActor");
+}
+#endif
diff --git a/Engine/Plugins/Experimental/Text3D/Source/Text3D/Private/Text3DComponent.cpp b/Engine/Plugins/Experimental/Text3D/Source/Text3D/Private/Text3DComponent.cpp
--- a/Engine/Plugins/Experimental/Text3D/Source/Text3D/Private/Text3DComponent.cpp
+++ b/Engine/Plugins/Experimental/Text3D/Source/Text3D/Private/Text3DComponent.cpp
@@ -6,6 +6,7 @@
 #include "Engine/Engine.h"
 #include "Engine/Font.h"
 #include "Engine/StaticMesh.h"
+#include "Extensions/Text3DDefaultCharacterExtension.h"
 #include "Extensions/Text3DDefaultGeometryExtension.h"
 #include "Extensions/Text3DDefaultLayoutExtension.h"
 #include "Extensions/Text3DDefaultMaterialExtension.h"
@@ -14,6 +15,7 @@
 #include "Internationalization/Regex.h"
 #include "Logs/Text3DLogs.h"
 #include "Materials/Material.h"
+#include "Misc/EnumerateRange.h"
 #include "Misc/ScopeExit.h"
 #include "Misc/TransactionObjectEvent.h"
 #include "Renderers/Text3DRendererBase.h"
@@ -21,9 +23,12 @@
 #include "Settings/Text3DProjectSettings.h"
 #include "Styling/StyleDefaults.h"
 #include "Subsystems/Text3DEngineSubsystem.h"
+#include "Text3DModule.h"
 #include "Text3DComponentVersion.h"
+#include "Text3DModule.h"
 #include "UObject/ConstructorHelpers.h"
 #include "UObject/UObjectGlobals.h"
+#include "Utilities/Text3DUtilities.h"
 
 #define LOCTEXT_NAMESPACE "Text3D"
 
@@ -51,6 +56,7 @@ UText3DComponent::UText3DComponent()
 	GeometryExtension = CreateDefaultSubobject<UText3DGeometryExtensionBase, UText3DDefaultGeometryExtension>(TEXT("GeometryExtension"));
 	LayoutExtension = CreateDefaultSubobject<UText3DLayoutExtensionBase, UText3DDefaultLayoutExtension>(TEXT("LayoutExtension"));
 	RenderingExtension = CreateDefaultSubobject<UText3DRenderingExtensionBase, UText3DDefaultRenderingExtension>(TEXT("RenderingExtension"));
+	CharacterExtension = CreateDefaultSubobject<UText3DCharacterExtensionBase, UText3DDefaultCharacterExtension>(TEXT("CharacterExtension"));
 }
 
 void UText3DComponent::OnComponentCreated()
@@ -106,7 +112,7 @@ void UText3DComponent::ClearUpdateFlags()
 
 uint32 UText3DComponent::GetTypeFaceIndex() const
 {
-	uint32 TypefaceIndex = 0;
+	uint32 TypefaceIndex = INDEX_NONE;
 
 	if (Font)
 	{
@@ -125,6 +131,24 @@ uint32 UText3DComponent::GetTypeFaceIndex() const
 	return TypefaceIndex;
 }
 
+const FTypefaceEntry* UText3DComponent::GetTypeFaceEntry() const
+{
+	if (Font)
+	{
+		TArray<FTypefaceEntry>& TypefaceEntries = Font->CompositeFont.DefaultTypeface.Fonts;
+
+		for (uint32 Index = 0; Index < static_cast<uint32>(TypefaceEntries.Num()); Index++)
+		{
+			if (Typeface == TypefaceEntries[Index].Name)
+			{
+				return &TypefaceEntries[Index];
+			}
+		}
+	}
+
+	return nullptr;
+}
+
 bool UText3DComponent::IsTypefaceAvailable(FName InTypeface) const
 {
 	for (const FTypefaceEntry& TypefaceElement : GetAvailableTypefaces())
@@ -369,7 +393,7 @@ void UText3DComponent::SetFontByName(const FString& InFontName)
 
 const FText& UText3DComponent::GetFormattedText() const
 {
-	return FormattedText.Get(FText::GetEmpty());
+	return FormattedText.Get(Text);
 }
 
 void UText3DComponent::SetText(const FText& Value)
@@ -394,17 +418,6 @@ void UText3DComponent::SetEnforceUpperCase(bool bInEnforceUpperCase)
 	OnTextChanged();
 }
 
-FText UText3DComponent::GetFormattedText()
-{
-	if (!FormattedText.IsSet())
-	{
-		FormattedText = Text;
-		FormatText(FormattedText.GetValue());
-	}
-
-	return FormattedText.GetValue();
-}
-
 void UText3DComponent::SetFont(UFont* InFont)
 {
 	if (!InFont || Font == InFont)
@@ -551,6 +564,25 @@ UText3DRendererBase* UText3DComponent::GetTextRenderer() const
 	return TextRenderer;
 }
 
+uint16 UText3DComponent::GetCharacterCount() const
+{
+	return CharacterExtension->GetCharacterCount();
+}
+
+UText3DCharacterBase* UText3DComponent::GetCharacter(uint16 InCharacterIndex) const
+{
+	return CharacterExtension->GetCharacter(InCharacterIndex);
+}
+
+void UText3DComponent::ForEachCharacter(const TFunctionRef<void(UText3DCharacterBase*, uint16, uint16)>& InFunctor) const
+{
+	const uint16 TotalCount = GetCharacterCount();
+	for (TConstEnumerateRef<UText3DCharacterBase*> CharacterRef : EnumerateRange(CharacterExtension->GetCharacters()))
+	{
+		InFunctor(*CharacterRef, CharacterRef.GetIndex(), TotalCount);
+	}
+}
+
 void UText3DComponent::SetKerning(const float Value)
 {
 	if (UText3DDefaultLayoutExtension* DefaultLayoutExtension = GetCastedLayoutExtension<UText3DDefaultLayoutExtension>())
@@ -828,11 +860,8 @@ void UText3DComponent::RebuildInternal(bool bCleanCache)
 
 	if (EnumHasAnyFlags(UpdateFlags,EText3DRendererFlags::Geometry))
 	{
-		if (!FormattedText.IsSet())
-		{
-			FormattedText = Text;
-			FormatText(FormattedText.GetValue());
-		}
+		FormattedText = Text;
+		FormatText(FormattedText.GetValue());
 
 		TextRenderer->Clear();
 
@@ -840,7 +869,7 @@ void UText3DComponent::RebuildInternal(bool bCleanCache)
 	}
 
 	TextRenderer->Update(UpdateFlags);
-	
+
 	if (bBroadcastEvent)
 	{
 		UpdateStatistics();
@@ -912,9 +941,6 @@ void UText3DComponent::OnTextRendererClassChanged()
 
 void UText3DComponent::OnTextChanged()
 {
-	FormattedText = Text;
-	FormatText(FormattedText.GetValue());
-	
 	RequestUpdate(EText3DRendererFlags::All);
 }
 
diff --git a/Engine/Plugins/Experimental/Text3D/Source/Text3D/Private/Text3DModule.h b/Engine/Plugins/Experimental/Text3D/Source/Text3D/Private/Text3DModule.h
--- a/Engine/Plugins/Experimental/Text3D/Source/Text3D/Private/Text3DModule.h
+++ b/Engine/Plugins/Experimental/Text3D/Source/Text3D/Private/Text3DModule.h
@@ -18,9 +18,6 @@ THIRD_PARTY_INCLUDES_END
 #undef generic
 #endif	//PLATFORM_COMPILER_HAS_GENERIC_KEYWORD
 
-const int32 FontSize = 64;
-const float FontInverseScale = 1.0f / FontSize;
-
 class FText3DModule final : public IModuleInterface
 {
 public:
diff --git a/Engine/Plugins/Experimental/Text3D/Source/Text3D/Private/Utilities/Text3DUtilities.cpp b/Engine/Plugins/Experimental/Text3D/Source/Text3D/Private/Utilities/Text3DUtilities.cpp
--- a/Engine/Plugins/Experimental/Text3D/Source/Text3D/Private/Utilities/Text3DUtilities.cpp
+++ b/Engine/Plugins/Experimental/Text3D/Source/Text3D/Private/Utilities/Text3DUtilities.cpp
@@ -4,6 +4,7 @@
 
 #include "Engine/Font.h"
 #include "Engine/FontFace.h"
+#include "Fonts/CompositeFont.h"
 #include "Misc/Paths.h"
 #include "Text3DModule.h"
 
@@ -207,3 +208,74 @@ bool UE::Text3D::Utilities::GetFontFaces(const UFont* InFont, TArray<UFontFace*>
 
 	return !OutFontFaces.IsEmpty();
 }
+
+TArray<FText3DCharacterMetadata> UE::Text3D::Utilities::GetCharacterMetadata(const FTypefaceEntry& InTypeFace, const FString& InText)
+{
+	const FFontFaceDataConstPtr FaceData = InTypeFace.Font.GetFontFaceData();
+
+	TArray<FText3DCharacterMetadata> CharacterIndexes;
+	CharacterIndexes.Reserve(InText.Len() + 1);
+
+	if (FaceData.IsValid() && FaceData->HasData() && !FaceData->GetData().IsEmpty())
+	{
+		const TArray<uint8> Data = FaceData->GetData();
+
+		FT_Face FreeTypeFace;
+		FT_New_Memory_Face(FText3DModule::GetFreeTypeLibrary(), Data.GetData(), Data.Num(), 0, &FreeTypeFace);
+
+		if (FreeTypeFace)
+		{
+			int32 TextIndex = 0;
+			int32 VisibleIndex = 0;
+			int32 LayoutIndex = 0;
+
+			FT_Set_Pixel_Sizes(FreeTypeFace, 0, UE::Text3D::Metrics::FontSize);
+
+			for (const TCHAR Character : InText)
+			{
+				if (!FChar::IsControl(Character) && FChar::IsPrint(Character))
+				{
+					if (!FChar::IsWhitespace(Character))
+					{
+						// Check if the font can render this character
+						FT_UInt GlyphIndex = FT_Get_Char_Index(FreeTypeFace, Character);
+        
+						if (GlyphIndex > 0)
+						{
+							// Load the glyph to check if it has an actual outline or bitmap
+							const FT_Error Code = FT_Load_Glyph(FreeTypeFace, GlyphIndex, FT_LOAD_DEFAULT);
+							if (Code == 0)
+							{
+								// Check if glyph has bitmap data or outline data
+								if ((FreeTypeFace->glyph->format == FT_GLYPH_FORMAT_BITMAP
+									&& FreeTypeFace->glyph->bitmap.width > 0 && FreeTypeFace->glyph->bitmap.rows > 0)
+									|| (FreeTypeFace->glyph->format == FT_GLYPH_FORMAT_OUTLINE && FreeTypeFace->glyph->outline.n_points > 0))
+								{
+									FText3DCharacterMetadata Metadata;
+									Metadata.GlyphIndex = GlyphIndex;
+									Metadata.LayoutIndex = LayoutIndex;
+									Metadata.VisibleIndex = VisibleIndex;
+									Metadata.TextIndex = TextIndex;
+									Metadata.Height = FreeTypeFace->glyph->metrics.height / UE::Text3D::Metrics::FontSize;
+									Metadata.Width = FreeTypeFace->glyph->metrics.width / UE::Text3D::Metrics::FontSize;
+
+									CharacterIndexes.Emplace(Metadata);
+
+									VisibleIndex++;
+								}
+							}
+						}
+					}
+
+					LayoutIndex++;
+				}
+
+				TextIndex++;
+			}
+	
+			FT_Done_Face(FreeTypeFace);
+		}
+	}
+
+	return CharacterIndexes;
+}
diff --git a/Engine/Plugins/Experimental/Text3D/Source/Text3D/Public/Characters/Text3DCharacterBase.h b/Engine/Plugins/Experimental/Text3D/Source/Text3D/Public/Characters/Text3DCharacterBase.h
new file mode 100644
--- /dev/null
+++ b/Engine/Plugins/Experimental/Text3D/Source/Text3D/Public/Characters/Text3DCharacterBase.h
@@ -0,0 +1,89 @@
+// Copyright Epic Games, Inc. All Rights Reserved.
+
+#pragma once
+
+#include "Misc/Optional.h"
+#include "Text3DTypes.h"
+#include "Text3DCharacterBase.generated.h"
+
+/** Holds data for a single character in Text3D */
+UCLASS(AutoExpandCategories=(Character))
+class UText3DCharacterBase : public UObject
+{
+	GENERATED_BODY()
+
+public:
+	TEXT3D_API FTransform& GetTransform(bool bInReset);
+
+#if WITH_EDITORONLY_DATA
+	void SetCharacter(const FString& InCharacter)
+	{
+		Character = InCharacter;
+	}
+#endif
+	
+	void SetCharacterMetadata(const FText3DCharacterMetadata& InMetadata);
+	const FText3DCharacterMetadata& GetCharacterMetadata() const
+	{
+		return Metadata;
+	}
+	FText3DCharacterMetadata& GetCharacterMetadata()
+	{
+		return Metadata;
+	}
+
+	void SetRelativeLocation(const FVector& InLocation);
+	const FVector& GetRelativeLocation() const
+	{
+		return RelativeLocation;
+	}
+
+	void SetRelativeRotation(const FRotator& InRotation);
+	const FRotator& GetRelativeRotation() const
+	{
+		return RelativeRotation;
+	}
+
+	void SetRelativeScale(const FVector& InScale);
+	const FVector& GetRelativeScale() const
+	{
+		return RelativeScale;
+	}
+	
+	/** Get character custom kerning if supported */
+	virtual TOptional<float> GetCharacterKerning() const
+	{
+		return TOptional<float>();
+	}
+
+protected:
+	//~ Begin UObject
+#if WITH_EDITOR
+	virtual void PostEditUndo() override;
+	virtual void PostEditChangeProperty(FPropertyChangedEvent& InEvent) override;
+#endif
+	//~ End UObject
+
+	void OnCharacterDataChanged() const;
+
+#if WITH_EDITORONLY_DATA
+	/** Used for TitleProperty of the array to display this instead of the index */
+	UPROPERTY(VisibleInstanceOnly, Category="Character", Transient, meta=(EditCondition="false", EditConditionHides, AllowPrivateAccess="true"))
+	FString Character;
+#endif
+
+	UPROPERTY(EditInstanceOnly, Getter, Setter, Category="Character", meta=(AllowPrivateAccess = "true"))
+	FVector RelativeLocation = FVector::ZeroVector;
+	
+	UPROPERTY(EditInstanceOnly, Getter, Setter, Category="Character", meta=(AllowPrivateAccess = "true"))
+	FRotator RelativeRotation = FRotator::ZeroRotator;
+	
+	UPROPERTY(EditInstanceOnly, Getter, Setter, Category="Character", meta=(Delta="0.05", MotionDesignVectorWidget, AllowPreserveRatio="XYZ", ClampMin="0", AllowPrivateAccess = "true"))
+	FVector RelativeScale = FVector::OneVector;
+
+	/** Final transform after all extensions are applied */
+	FTransform Transform = FTransform::Identity;
+
+	/** Character metadata */
+	FText3DCharacterMetadata Metadata;
+};
diff --git a/Engine/Plugins/Experimental/Text3D/Source/Text3D/Public/Characters/Text3DDefaultCharacter.h b/Engine/Plugins/Experimental/Text3D/Source/Text3D/Public/Characters/Text3DDefaultCharacter.h
new file mode 100644
--- /dev/null
+++ b/Engine/Plugins/Experimental/Text3D/Source/Text3D/Public/Characters/Text3DDefaultCharacter.h
@@ -0,0 +1,47 @@
+// Copyright Epic Games, Inc. All Rights Reserved.
+
+#pragma once
+
+#include "Text3DCharacterBase.h"
+#include "Text3DTypes.h"
+#include "Text3DDefaultCharacter.generated.h"
+
+/** Holds data for a single character in Text3D */
+UCLASS(MinimalAPI, AutoExpandCategories=(Character))
+class UText3DDefaultCharacter : public UText3DCharacterBase
+{
+	GENERATED_BODY()
+
+public:
+	static TEXT3D_API FName GetKerningPropertyName();
+
+	TEXT3D_API void SetKerningOverride(bool bInKerningOverride);
+	bool GetKerningOverride() const
+	{
+		return bKerningOverride;
+	}
+
+	TEXT3D_API void SetKerning(float InKerning);
+	float GetKerning() const
+	{
+		return Kerning;
+	}
+
+protected:
+	//~ Begin UObject
+#if WITH_EDITOR
+	virtual void PostEditChangeProperty(FPropertyChangedEvent& InEvent) override;
+#endif
+	//~ End UObject
+
+	//~ Begin UText3DCharacterBase
+	virtual TOptional<float> GetCharacterKerning() const override;
+	//~ End UText3DCharacterBase
+
+	UPROPERTY(EditInstanceOnly, Getter="GetKerningOverride", Setter="SetKerningOverride", Category="Character", meta=(InlineEditConditionToggle, AllowPrivateAccess = "true"))
+	bool bKerningOverride = false;
+
+	/** Kerning adjusts the space between this character */
+	UPROPERTY(EditInstanceOnly, Getter, Setter, Category="Character", meta=(EditCondition="bKerningOverride", AllowPrivateAccess = "true"))
+	float Kerning;
+};
diff --git a/Engine/Plugins/Experimental/Text3D/Source/Text3D/Public/Extensions/Text3DCharacterExtensionBase.h b/Engine/Plugins/Experimental/Text3D/Source/Text3D/Public/Extensions/Text3DCharacterExtensionBase.h
new file mode 100644
--- /dev/null
+++ b/Engine/Plugins/Experimental/Text3D/Source/Text3D/Public/Extensions/Text3DCharacterExtensionBase.h
@@ -0,0 +1,39 @@
+// Copyright Epic Games, Inc. All Rights Reserved.
+
+#pragma once
+
+#include "Text3DExtensionBase.h"
+#include "Text3DTypes.h"
+#include "Text3DCharacterExtensionBase.generated.h"
+
+class UText3DCharacterBase;
+
+/** Extension that handles character data for Text3D */
+UCLASS(MinimalAPI, Abstract)
+class UText3DCharacterExtensionBase : public UText3DExtensionBase
+{
+	GENERATED_BODY()
+
+public:
+	UText3DCharacterExtensionBase()
+		: UText3DExtensionBase(0)
+	{}
+
+	/** Get the amount of character */
+	virtual uint16 GetCharacterCount() const
+	{
+		return 0;
+	}
+
+	/** Get per character data */
+	virtual UText3DCharacterBase* GetCharacter(uint16 InCharacter) const
+	{
+		return nullptr;
+	}
+
+	/** Get all character data */
+	virtual TConstArrayView<UText3DCharacterBase*> GetCharacters() const
+	{
+		return {};
+	}
+};
\ No newline at end of file
diff --git a/Engine/Plugins/Experimental/Text3D/Source/Text3D/Public/Extensions/Text3DDefaultCharacterExtension.h b/Engine/Plugins/Experimental/Text3D/Source/Text3D/Public/Extensions/Text3DDefaultCharacterExtension.h
new file mode 100644
--- /dev/null
+++ b/Engine/Plugins/Experimental/Text3D/Source/Text3D/Public/Extensions/Text3DDefaultCharacterExtension.h
@@ -0,0 +1,38 @@
+// Copyright Epic Games, Inc. All Rights Reserved.
+
+#pragma once
+
+#include "Text3DCharacterExtensionBase.h"
+#include "Subsystems/Text3DEngineSubsystem.h"
+#include "Text3DDefaultCharacterExtension.generated.h"
+
+class UText3DDefaultCharacter;
+
+UCLASS(MinimalAPI)
+class UText3DDefaultCharacterExtension : public UText3DCharacterExtensionBase
+{
+	GENERATED_BODY()
+	
+protected:
+	//~ Begin UText3DGeometryExtensionBase
+	virtual uint16 GetCharacterCount() const override;
+	virtual UText3DCharacterBase* GetCharacter(uint16 InCharacter) const override;
+	virtual TConstArrayView<UText3DCharacterBase*> GetCharacters() const override;
+	//~ End UText3DGeometryExtensionBase
+
+	//~ Begin UText3DExtensionBase
+	virtual EText3DExtensionResult PreRendererUpdate(EText3DRendererFlags InFlag) override;
+	virtual EText3DExtensionResult PostRendererUpdate(EText3DRendererFlags InFlag) override;
+	//~ End UText3DExtensionBase
+
+	/** Allocate character data by using pool or creating new object */
+	void AllocateTextCharacters(uint16 InCharacterCount);
+
+	/** Characters composing the active text */
+	UPROPERTY(VisibleAnywhere, EditFixedSize, Instanced, Category = "Character", meta = (TitleProperty="Character", EditFixedOrder, AllowPrivateAccess = "true"))
+	TArray<TObjectPtr<UText3DCharacterBase>> TextCharacters;
+
+	/** Pool of characters to reuse when text changes to avoid creating new objects */
+	UPROPERTY(Transient, DuplicateTransient, TextExportTransient)
+	TArray<TObjectPtr<UText3DCharacterBase>> TextCharactersPool;
+};
\ No newline at end of file
diff --git a/Engine/Plugins/Experimental/Text3D/Source/Text3D/Public/Extensions/Text3DDefaultGeometryExtension.h b/Engine/Plugins/Experimental/Text3D/Source/Text3D/Public/Extensions/Text3DDefaultGeometryExtension.h
--- a/Engine/Plugins/Experimental/Text3D/Source/Text3D/Public/Extensions/Text3DDefaultGeometryExtension.h
+++ b/Engine/Plugins/Experimental/Text3D/Source/Text3D/Public/Extensions/Text3DDefaultGeometryExtension.h
@@ -6,8 +6,6 @@
 #include "Subsystems/Text3DEngineSubsystem.h"
 #include "Text3DDefaultGeometryExtension.generated.h"
 
-struct FPropertyChangedEvent;
-
 UCLASS(MinimalAPI)
 class UText3DDefaultGeometryExtension : public UText3DGeometryExtensionBase
 {
@@ -79,6 +77,22 @@ public:
 	/** Set the outline width. */
 	UFUNCTION(BlueprintCallable, Category = "Text3D|Geometry")
 	void SetOutline(const float Value);
+
+	UFUNCTION(BlueprintCallable, Category = "Text3D|Geometry")
+	void SetPivotHAlignment(EText3DHorizontalTextAlignment InPivot);
+
+	EText3DHorizontalTextAlignment GetPivotHAlignment() const
+	{
+		return PivotHAlignment;
+	}
+	
+	UFUNCTION(BlueprintCallable, Category = "Text3D|Geometry")
+	void SetPivotVAlignment(EText3DVerticalTextAlignment InPivot);
+
+	EText3DVerticalTextAlignment GetPivotVAlignment() const
+	{
+		return PivotVAlignment;
+	}
 	
 protected:
 	//~ Begin UObject
@@ -88,7 +102,9 @@ protected:
 	//~ End UObject
 
 	//~ Begin UText3DGeometryExtensionBase
-	virtual UStaticMesh* GetGlyphMesh(uint32 InGlyphIndex, const FFreeTypeFace* InFontFace) const override;
+	virtual UStaticMesh* GetGlyphMesh(UText3DCharacterBase* InCharacter) const override;
+	virtual EText3DHorizontalTextAlignment GetGlyphHAlignment() const override;
+	virtual EText3DVerticalTextAlignment GetGlyphVAlignment() const override;
 	//~ End UText3DGeometryExtensionBase
 
 	//~ Begin UText3DExtensionBase
@@ -99,6 +115,7 @@ protected:
 	void OnGeometryOptionsChanged();
 
 	float GetMaxBevel() const;
+	FVector GetPivotOffset() const;
 
 	/** Size of the extrude */
 	UPROPERTY(EditAnywhere, Getter, Setter, Category = "Geometry", meta = (ClampMin = 0, AllowPrivateAccess = "true"))
@@ -124,6 +141,14 @@ protected:
 	UPROPERTY(EditAnywhere, Getter, Setter, Category = "Geometry", meta = (EditCondition = "bUseOutline", AllowPrivateAccess = "true"))
 	float Outline = 0.5f;
 
+	/** Pivot alignment of each character */
+	UPROPERTY(EditAnywhere, Getter, Setter, Category = "Geometry", meta = (AllowPrivateAccess = "true"))
+	EText3DHorizontalTextAlignment PivotHAlignment = EText3DHorizontalTextAlignment::Left;
+	
+	/** Pivot alignment of each character */
+	UPROPERTY(EditAnywhere, Getter, Setter, Category = "Geometry", meta = (AllowPrivateAccess = "true", ValidEnumValues="Bottom"))
+	EText3DVerticalTextAlignment PivotVAlignment = EText3DVerticalTextAlignment::Bottom;
+
 private:
 	FCachedFontData* CachedFontData = nullptr;
 	FGlyphMeshParameters GlyphMeshParameters;
diff --git a/Engine/Plugins/Experimental/Text3D/Source/Text3D/Public/Extensions/Text3DDefaultLayoutExtension.h b/Engine/Plugins/Experimental/Text3D/Source/Text3D/Public/Extensions/Text3DDefaultLayoutExtension.h
--- a/Engine/Plugins/Experimental/Text3D/Source/Text3D/Public/Extensions/Text3DDefaultLayoutExtension.h
+++ b/Engine/Plugins/Experimental/Text3D/Source/Text3D/Public/Extensions/Text3DDefaultLayoutExtension.h
@@ -6,6 +6,8 @@
 #include "Framework/Text/ITextLayoutMarshaller.h"
 #include "Text3DDefaultLayoutExtension.generated.h"
 
+class FText3DLayout;
+struct FTextBlockStyle;
 struct FPropertyChangedEvent;
 
 UCLASS(MinimalAPI)
@@ -142,8 +144,6 @@ public:
 	void SetScaleProportionally(const bool Value);
 
 protected:
-	UText3DDefaultLayoutExtension();
-
 	//~ Begin UObject
 #if WITH_EDITOR
 	virtual void PostEditChangeProperty(FPropertyChangedEvent& InEvent) override;
@@ -154,7 +154,6 @@ protected:
 	virtual float GetTextHeight() const override;
 	virtual FVector GetTextScale() const override;
 	virtual FVector GetLineLocation(int32 LineIndex) const override;
-	virtual TSharedPtr<FText3DShapedGlyphText> GetShapedText() const override;
 	//~ End UText3DLayoutExtensionBase
 
 	//~ Begin UText3DExtensionBase
@@ -166,11 +165,11 @@ protected:
 	void OnLayoutOptionsChanged();
 
 	/** Horizontal text alignment */
-	UPROPERTY(EditAnywhere, Getter, Setter, Category = "Layout", meta = (AllowPrivateAccess = "true"))
+	UPROPERTY(EditAnywhere, Getter, Setter, Category = "Layout", meta = (AllowPrivateAccess = "true", CustomAlignmentWidget))
 	EText3DHorizontalTextAlignment HorizontalAlignment = EText3DHorizontalTextAlignment::Left;
 
 	/** Vertical text alignment */
-	UPROPERTY(EditAnywhere, Getter, Setter, Category = "Layout", meta = (AllowPrivateAccess = "true"))
+	UPROPERTY(EditAnywhere, Getter, Setter, Category = "Layout", meta = (AllowPrivateAccess = "true", CustomAlignmentWidget))
 	EText3DVerticalTextAlignment VerticalAlignment = EText3DVerticalTextAlignment::FirstLine;
 
 	/** Text kerning */
@@ -217,7 +216,9 @@ private:
 	TSharedPtr<FText3DShapedGlyphText> ShapedText;
 
 	/** Stores the text layout calculated by the TextLayoutMarshaller. */
-	TSharedPtr<FTextLayout> TextLayout;
+	TSharedPtr<FText3DLayout> TextLayout;
+
+	TSharedPtr<FTextBlockStyle> FontStyle;
 
 	/** Determines how text is laid out, ie. parsing line breaks. */
 	TSharedPtr<ITextLayoutMarshaller> TextLayoutMarshaller;
diff --git a/Engine/Plugins/Experimental/Text3D/Source/Text3D/Public/Extensions/Text3DDefaultTransformExtension.h b/Engine/Plugins/Experimental/Text3D/Source/Text3D/Public/Extensions/Text3DDefaultTransformExtension.h
--- a/Engine/Plugins/Experimental/Text3D/Source/Text3D/Public/Extensions/Text3DDefaultTransformExtension.h
+++ b/Engine/Plugins/Experimental/Text3D/Source/Text3D/Public/Extensions/Text3DDefaultTransformExtension.h
@@ -172,16 +172,16 @@ protected:
 	UPROPERTY(EditAnywhere, Getter="GetLocationEnabled", Setter="SetLocationEnabled", Category = "Transform")
 	bool bLocationEnabled = false;
 
-	UPROPERTY(EditAnywhere, Getter, Setter, Category = "Transform", meta = (EditCondition = "bLocationEnabled", ClampMin = 0, ClampMax = 100))
+	UPROPERTY(EditAnywhere, Getter, Setter, Category = "Transform", meta = (EditCondition = "bLocationEnabled", EditConditionHides, ClampMin = 0, ClampMax = 100))
 	float LocationProgress = 0.f;
 
-	UPROPERTY(EditAnywhere, Getter, Setter, Category = "Transform", meta = (EditCondition = "bLocationEnabled"))
+	UPROPERTY(EditAnywhere, Getter, Setter, Category = "Transform", meta = (EditCondition = "bLocationEnabled", EditConditionHides))
 	EText3DCharacterEffectOrder LocationOrder = EText3DCharacterEffectOrder::Normal;
 
-	UPROPERTY(EditAnywhere, Getter, Setter, Category = "Transform", meta = (EditCondition = "bLocationEnabled", ClampMin = 0, ClampMax = 100))
+	UPROPERTY(EditAnywhere, Getter, Setter, Category = "Transform", meta = (EditCondition = "bLocationEnabled", EditConditionHides, ClampMin = 0, ClampMax = 100))
 	float LocationRange = 50.f;
 
-	UPROPERTY(EditAnywhere, Getter, Setter, Category = "Transform", meta = (EditCondition = "bLocationEnabled"))
+	UPROPERTY(EditAnywhere, Getter, Setter, Category = "Transform", meta = (EditCondition = "bLocationEnabled", EditConditionHides))
 	FVector LocationDistance = FVector(100.f, 0.f, 0.f);
 
 	// Rotate
@@ -189,19 +189,19 @@ protected:
 	UPROPERTY(EditAnywhere, Getter="GetRotationEnabled", Setter="SetRotationEnabled", Category = "Transform")
 	bool bRotationEnabled = false;
 
-	UPROPERTY(EditAnywhere, Getter, Setter, Category = "Transform", meta = (EditCondition = "bRotationEnabled", ClampMin = 0, ClampMax = 100))
+	UPROPERTY(EditAnywhere, Getter, Setter, Category = "Transform", meta = (EditCondition = "bRotationEnabled", EditConditionHides, ClampMin = 0, ClampMax = 100))
 	float RotationProgress = 0.0f;
 
-	UPROPERTY(EditAnywhere, Getter, Setter, Category = "Transform", meta = (EditCondition = "bRotationEnabled"))
+	UPROPERTY(EditAnywhere, Getter, Setter, Category = "Transform", meta = (EditCondition = "bRotationEnabled", EditConditionHides))
 	EText3DCharacterEffectOrder RotationOrder = EText3DCharacterEffectOrder::Normal;
 
-	UPROPERTY(EditAnywhere, Getter, Setter, Category = "Transform", meta = (EditCondition = "bRotationEnabled", ClampMin = 0, ClampMax = 100))
+	UPROPERTY(EditAnywhere, Getter, Setter, Category = "Transform", meta = (EditCondition = "bRotationEnabled", EditConditionHides, ClampMin = 0, ClampMax = 100))
 	float RotationRange = 50.f;
 
-	UPROPERTY(EditAnywhere, Getter, Setter, Category = "Transform", meta = (EditCondition = "bRotationEnabled"))
+	UPROPERTY(EditAnywhere, Getter, Setter, Category = "Transform", meta = (EditCondition = "bRotationEnabled", EditConditionHides))
 	FRotator RotationBegin = FRotator(-90.f, 0.f, 0.f);
 
-	UPROPERTY(EditAnywhere, Getter, Setter, Category = "Transform", meta = (EditCondition = "bRotationEnabled"))
+	UPROPERTY(EditAnywhere, Getter, Setter, Category = "Transform", meta = (EditCondition = "bRotationEnabled", EditConditionHides))
 	FRotator RotationEnd = FRotator(0.f, 0.f, 0.f);
 
 	// Scale
@@ -209,18 +209,18 @@ protected:
 	UPROPERTY(EditAnywhere, Getter="GetScaleEnabled", Setter="SetScaleEnabled", Category = "Transform")
 	bool bScaleEnabled = false;
 
-	UPROPERTY(EditAnywhere, Getter, Setter, Category = "Transform", meta = (EditCondition = "bScaleEnabled", ClampMin = 0, ClampMax = 100))
+	UPROPERTY(EditAnywhere, Getter, Setter, Category = "Transform", meta = (EditCondition = "bScaleEnabled", EditConditionHides, ClampMin = 0, ClampMax = 100))
 	float ScaleProgress = 0.f;
 
-	UPROPERTY(EditAnywhere, Getter, Setter, Category = "Transform", meta = (EditCondition = "bScaleEnabled"))
+	UPROPERTY(EditAnywhere, Getter, Setter, Category = "Transform", meta = (EditCondition = "bScaleEnabled", EditConditionHides))
 	EText3DCharacterEffectOrder ScaleOrder = EText3DCharacterEffectOrder::Normal;
 
-	UPROPERTY(EditAnywhere, Getter, Setter, Category = "Transform", meta = (EditCondition = "bScaleEnabled", ClampMin = 0, ClampMax = 100))
+	UPROPERTY(EditAnywhere, Getter, Setter, Category = "Transform", meta = (EditCondition = "bScaleEnabled", EditConditionHides, ClampMin = 0, ClampMax = 100))
 	float ScaleRange = 50.f;
 
-	UPROPERTY(EditAnywhere, Getter, Setter, Category = "Transform", meta = (EditCondition = "bScaleEnabled", ClampMin = 0))
+	UPROPERTY(EditAnywhere, Getter, Setter, Category = "Transform", meta = (EditCondition = "bScaleEnabled", EditConditionHides, ClampMin = 0))
 	FVector ScaleBegin = FVector(1.f, 0.f, 0.f);
 
-	UPROPERTY(EditAnywhere, Getter, Setter, Category = "Transform", meta = (EditCondition = "bScaleEnabled", ClampMin = 0))
+	UPROPERTY(EditAnywhere, Getter, Setter, Category = "Transform", meta = (EditCondition = "bScaleEnabled", EditConditionHides, ClampMin = 0))
 	FVector ScaleEnd = FVector::OneVector;
 };
diff --git a/Engine/Plugins/Experimental/Text3D/Source/Text3D/Public/Extensions/Text3DExtensionBase.h b/Engine/Plugins/Experimental/Text3D/Source/Text3D/Public/Extensions/Text3DExtensionBase.h
--- a/Engine/Plugins/Experimental/Text3D/Source/Text3D/Public/Extensions/Text3DExtensionBase.h
+++ b/Engine/Plugins/Experimental/Text3D/Source/Text3D/Public/Extensions/Text3DExtensionBase.h
@@ -7,7 +7,7 @@
 
 class UText3DComponent;
 
-enum class EText3DExtensionResult
+enum class EText3DExtensionResult : uint8
 {
 	/** Extension is still needed, keep it around */
 	Active,
@@ -27,6 +27,14 @@ class UText3DExtensionBase : public UObject
 	GENERATED_BODY()
 
 public:
+	UText3DExtensionBase()
+		: UText3DExtensionBase(0)
+	{}
+
+	UText3DExtensionBase(uint16 InPriority)
+		: UpdatePriority(InPriority)
+	{}
+
 	/** Perform an update of the extension behavior before renderer execute based on condition above */
 	virtual EText3DExtensionResult PreRendererUpdate(EText3DRendererFlags InFlag)
 	{
@@ -39,7 +47,16 @@ public:
 		return EText3DExtensionResult::Finished;
 	}
 
+	/** Used to sort extensions and execute them before others */
+	uint16 GetUpdatePriority() const
+	{
+		return UpdatePriority;
+	}
+
 protected:
 	UText3DComponent* GetText3DComponent() const;
 	void RequestUpdate(EText3DRendererFlags InFlags, bool bInImmediate = false) const;
+
+private:
+	uint16 UpdatePriority = 0;
 };
diff --git a/Engine/Plugins/Experimental/Text3D/Source/Text3D/Public/Extensions/Text3DGeometryExtensionBase.h b/Engine/Plugins/Experimental/Text3D/Source/Text3D/Public/Extensions/Text3DGeometryExtensionBase.h
--- a/Engine/Plugins/Experimental/Text3D/Source/Text3D/Public/Extensions/Text3DGeometryExtensionBase.h
+++ b/Engine/Plugins/Experimental/Text3D/Source/Text3D/Public/Extensions/Text3DGeometryExtensionBase.h
@@ -8,6 +8,7 @@
 
 class FFreeTypeFace;
 class UStaticMesh;
+class UText3DCharacterBase;
 
 /** Extension that handles geometry data for Text3D */
 UCLASS(MinimalAPI, Abstract)
@@ -16,8 +17,22 @@ class UText3DGeometryExtensionBase : public UText3DExtensionBase
 	GENERATED_BODY()
 
 public:
-	virtual UStaticMesh* GetGlyphMesh(uint32 InGlyphIndex, const FFreeTypeFace* InFontFace) const
+	UText3DGeometryExtensionBase()
+		: UText3DExtensionBase(/** Priority */1)
+	{}
+
+	virtual UStaticMesh* GetGlyphMesh(UText3DCharacterBase* InCharacter) const
 	{
 		return nullptr;
 	}
+
+	virtual EText3DHorizontalTextAlignment GetGlyphHAlignment() const
+	{
+		return EText3DHorizontalTextAlignment::Left;
+	}
+
+	virtual EText3DVerticalTextAlignment GetGlyphVAlignment() const
+	{
+		return EText3DVerticalTextAlignment::Bottom;
+	}
 };
diff --git a/Engine/Plugins/Experimental/Text3D/Source/Text3D/Public/Extensions/Text3DLayoutExtensionBase.h b/Engine/Plugins/Experimental/Text3D/Source/Text3D/Public/Extensions/Text3DLayoutExtensionBase.h
--- a/Engine/Plugins/Experimental/Text3D/Source/Text3D/Public/Extensions/Text3DLayoutExtensionBase.h
+++ b/Engine/Plugins/Experimental/Text3D/Source/Text3D/Public/Extensions/Text3DLayoutExtensionBase.h
@@ -15,6 +15,10 @@ class UText3DLayoutExtensionBase : public UText3DExtensionBase
 	GENERATED_BODY()
 
 public:
+	UText3DLayoutExtensionBase()
+		: UText3DExtensionBase(/** Priority */2)
+	{}
+
 	virtual float GetTextHeight() const
 	{
 		return 0.f;
@@ -29,9 +33,4 @@ public:
 	{
 		return FVector::ZeroVector;
 	}
-	
-	virtual TSharedPtr<FText3DShapedGlyphText> GetShapedText() const
-	{
-		return nullptr;
-	}
 };
diff --git a/Engine/Plugins/Experimental/Text3D/Source/Text3D/Public/Extensions/Text3DMaterialExtensionBase.h b/Engine/Plugins/Experimental/Text3D/Source/Text3D/Public/Extensions/Text3DMaterialExtensionBase.h
--- a/Engine/Plugins/Experimental/Text3D/Source/Text3D/Public/Extensions/Text3DMaterialExtensionBase.h
+++ b/Engine/Plugins/Experimental/Text3D/Source/Text3D/Public/Extensions/Text3DMaterialExtensionBase.h
@@ -15,6 +15,10 @@ class UText3DMaterialExtensionBase : public UText3DExtensionBase
 	GENERATED_BODY()
 
 public:
+	UText3DMaterialExtensionBase()
+		: UText3DExtensionBase(/** Priority */3)
+	{}
+
 	/** Set the material for a specific group */
 	virtual void SetMaterial(EText3DGroupType InGroup, UMaterialInterface* InMaterial)
 	{
diff --git a/Engine/Plugins/Experimental/Text3D/Source/Text3D/Public/Extensions/Text3DRenderingExtensionBase.h b/Engine/Plugins/Experimental/Text3D/Source/Text3D/Public/Extensions/Text3DRenderingExtensionBase.h
--- a/Engine/Plugins/Experimental/Text3D/Source/Text3D/Public/Extensions/Text3DRenderingExtensionBase.h
+++ b/Engine/Plugins/Experimental/Text3D/Source/Text3D/Public/Extensions/Text3DRenderingExtensionBase.h
@@ -13,6 +13,10 @@ class UText3DRenderingExtensionBase : public UText3DExtensionBase
 	GENERATED_BODY()
 
 public:
+	UText3DRenderingExtensionBase()
+		: UText3DExtensionBase(/** Priority */4)
+	{}
+
 	virtual void SetTextCastShadow(bool bInValue)
 	{
 	}
diff --git a/Engine/Plugins/Experimental/Text3D/Source/Text3D/Public/Extensions/Text3DTransformExtensionBase.h b/Engine/Plugins/Experimental/Text3D/Source/Text3D/Public/Extensions/Text3DTransformExtensionBase.h
--- a/Engine/Plugins/Experimental/Text3D/Source/Text3D/Public/Extensions/Text3DTransformExtensionBase.h
+++ b/Engine/Plugins/Experimental/Text3D/Source/Text3D/Public/Extensions/Text3DTransformExtensionBase.h
@@ -51,6 +51,10 @@ class UText3DTransformExtensionBase : public UText3DExtensionBase
 	GENERATED_BODY()
 
 public:
+	UText3DTransformExtensionBase()
+		: UText3DExtensionBase(/** Priority */2)
+	{}
+
 	virtual FText3DRange GetTargetRange() const
 	{
 		return FText3DRange();
diff --git a/Engine/Plugins/Experimental/Text3D/Source/Text3D/Public/Text3DActor.h b/Engine/Plugins/Experimental/Text3D/Source/Text3D/Public/Text3DActor.h
--- a/Engine/Plugins/Experimental/Text3D/Source/Text3D/Public/Text3DActor.h
+++ b/Engine/Plugins/Experimental/Text3D/Source/Text3D/Public/Text3DActor.h
@@ -16,7 +16,16 @@ public:
 	AText3DActor();
 
 	/** Returns Text3D subobject **/
-	UText3DComponent* GetText3DComponent() const { return Text3DComponent; }
+	UText3DComponent* GetText3DComponent() const
+	{
+		return Text3DComponent;
+	}
+
+	//~ Begin AActor
+#if WITH_EDITOR
+	virtual FString GetDefaultActorLabel() const override;
+#endif
+	//~ End AActor
 
 private:
 	UPROPERTY(Category = Text, VisibleAnywhere, BlueprintReadOnly, meta = (AllowPrivateAccess = "true"))
diff --git a/Engine/Plugins/Experimental/Text3D/Source/Text3D/Public/Text3DComponent.h b/Engine/Plugins/Experimental/Text3D/Source/Text3D/Public/Text3DComponent.h
--- a/Engine/Plugins/Experimental/Text3D/Source/Text3D/Public/Text3DComponent.h
+++ b/Engine/Plugins/Experimental/Text3D/Source/Text3D/Public/Text3DComponent.h
@@ -8,6 +8,8 @@
 #include "UObject/ObjectMacros.h"
 #include "Text3DComponent.generated.h"
 
+class UText3DCharacterBase;
+class UText3DCharacterExtensionBase;
 class FTextLayout;
 class ITextLayoutMarshaller;
 class UFont;
@@ -21,13 +23,17 @@ class UText3DTransformExtensionBase;
 struct FText3DShapedGlyphText;
 struct FTypefaceEntry;
 
-UCLASS(ClassGroup = (Text3D), PrioritizeCategories = "Text Layout Geometry Materials", meta = (BlueprintSpawnableComponent))
-class TEXT3D_API UText3DComponent : public USceneComponent
+UCLASS(MinimalAPI
+	, ClassGroup = (Text3D)
+	, PrioritizeCategories="Text Character Layout Geometry Material Rendering Transform"
+	, HideCategories=(Replication,Collision,HLOD,Physics,Networking,Input,Actor,Cooking,LevelInstance,Streaming,DataLayers,WorldPartition)
+	, meta = (BlueprintSpawnableComponent))
+class UText3DComponent : public USceneComponent
 {
 	GENERATED_BODY()
 
 public:
-	UText3DComponent();
+	TEXT3D_API UText3DComponent();
 
 	/**
 	 * Delegate called after text is rebuilt
@@ -60,27 +66,25 @@ public:
 		return Text;
 	}
 
-	const FText& GetFormattedText() const;
-
 	/** Set the text value and signal the primitives to be rebuilt */
 	UFUNCTION(BlueprintCallable, Category = "Rendering|Components|Text3D")
-	void SetText(const FText& Value);
+	TEXT3D_API void SetText(const FText& Value);
+
+	/**
+	 * Returns the Text property, after being formatted by the FormatText virtual function.
+	 * If FormatText is not overriden, the return FText will be the same as the Text property.
+	 */
+	UFUNCTION(BlueprintCallable, Category = "Text3D")
+	TEXT3D_API const FText& GetFormattedText() const;
 
 	UFUNCTION(BlueprintCallable, Category = "Rendering|Components|Text3D")
-	void SetEnforceUpperCase(bool bInEnforceUpperCase);
+	TEXT3D_API void SetEnforceUpperCase(bool bInEnforceUpperCase);
 
 	bool GetEnforceUpperCase() const
 	{
 		return bEnforceUpperCase;
 	}
 
-	/**
-	 * Returns the Text property, after being formatted by the FormatText virtual function.
-	 * If FormatText is not overriden, the return FText will be the same as the Text property.
-	 */
-	UFUNCTION(BlueprintCallable, Category = "Text3D")
-	FText GetFormattedText();
-
 	/** Get the text font and signal the primitives to be rebuilt */
 	UFont* GetFont() const
 	{
@@ -89,147 +93,147 @@ public:
 
 	/** Set the text font and signal the primitives to be rebuilt */
 	UFUNCTION(BlueprintCallable, Category = "Rendering|Components|Text3D")
-	void SetFont(UFont* InFont);
+	TEXT3D_API void SetFont(UFont* InFont);
 
 	/** Set whether an outline is applied. */
 	UFUNCTION(BlueprintCallable, Category = "Rendering|Components|Text3D")
-	void SetHasOutline(const bool bValue);
+	TEXT3D_API void SetHasOutline(const bool bValue);
 
 	/** Set the outline width. */
 	UFUNCTION(BlueprintCallable, Category = "Rendering|Components|Text3D")
-	void SetOutlineExpand(const float Value);
+	TEXT3D_API void SetOutlineExpand(const float Value);
 
 	/** Set the text extrusion size and signal the primitives to be rebuilt */
 	UFUNCTION(BlueprintCallable, Category = "Rendering|Components|Text3D")
-	void SetExtrude(const float Value);
+	TEXT3D_API void SetExtrude(const float Value);
 
 	/** Set the 3d bevel value */
 	UFUNCTION(BlueprintCallable, Category = "Rendering|Components|Text3D")
-	void SetBevel(const float Value);
+	TEXT3D_API void SetBevel(const float Value);
 
 	/** Set the 3d bevel type */
 	UFUNCTION(BlueprintCallable, Category = "Rendering|Components|Text3D")
-	void SetBevelType(const EText3DBevelType Value);
+	TEXT3D_API void SetBevelType(const EText3DBevelType Value);
 
 	/** Set the amount of segments that will be used to tessellate the Bevel */
 	UFUNCTION(BlueprintCallable, Category = "Rendering|Components|Text3D")
-	void SetBevelSegments(const int32 Value);
+	TEXT3D_API void SetBevelSegments(const int32 Value);
 
 	/** Get the text front material */
-	UMaterialInterface* GetFrontMaterial() const;
+	TEXT3D_API UMaterialInterface* GetFrontMaterial() const;
 
 	/** Set the text front material */
 	UFUNCTION(BlueprintCallable, Category = "Rendering|Components|Text3D")
-	void SetFrontMaterial(UMaterialInterface* Value);
+	TEXT3D_API void SetFrontMaterial(UMaterialInterface* Value);
 
 	/** Get the text bevel material */
-	UMaterialInterface* GetBevelMaterial() const;
+	TEXT3D_API UMaterialInterface* GetBevelMaterial() const;
 
 	/** Set the text bevel material */
 	UFUNCTION(BlueprintCallable, Category = "Rendering|Components|Text3D")
-	void SetBevelMaterial(UMaterialInterface* Value);
+	TEXT3D_API void SetBevelMaterial(UMaterialInterface* Value);
 	
 	/** Get the text extrude material */
-	UMaterialInterface* GetExtrudeMaterial() const;
+	TEXT3D_API UMaterialInterface* GetExtrudeMaterial() const;
 
 	/** Set the text extrude material */
 	UFUNCTION(BlueprintCallable, Category = "Rendering|Components|Text3D")
-	void SetExtrudeMaterial(UMaterialInterface* Value);
+	TEXT3D_API void SetExtrudeMaterial(UMaterialInterface* Value);
 
 	/** Get the text back material */
-	UMaterialInterface* GetBackMaterial() const;
+	TEXT3D_API UMaterialInterface* GetBackMaterial() const;
 
 	/** Set the text back material */
 	UFUNCTION(BlueprintCallable, Category = "Rendering|Components|Text3D")
-	void SetBackMaterial(UMaterialInterface* Value);
+	TEXT3D_API void SetBackMaterial(UMaterialInterface* Value);
 
 	/** Get the kerning value and signal the primitives to be rebuilt */
-	float GetKerning() const;
+	TEXT3D_API float GetKerning() const;
 
 	/** Set the kerning value and signal the primitives to be rebuilt */
 	UFUNCTION(BlueprintCallable, Category = "Rendering|Components|Text3D")
-	void SetKerning(const float Value);
+	TEXT3D_API void SetKerning(const float Value);
 
 	/** Get the line spacing value and signal the primitives to be rebuilt */
-	float GetLineSpacing() const;
+	TEXT3D_API float GetLineSpacing() const;
 
 	/** Set the line spacing value and signal the primitives to be rebuilt */
 	UFUNCTION(BlueprintCallable, Category = "Rendering|Components|Text3D")
-	void SetLineSpacing(const float Value);
+	TEXT3D_API void SetLineSpacing(const float Value);
 
 	/** Get the word spacing value and signal the primitives to be rebuilt */
-	float GetWordSpacing() const;
+	TEXT3D_API float GetWordSpacing() const;
 
 	/** Set the word spacing value and signal the primitives to be rebuilt */
 	UFUNCTION(BlueprintCallable, Category = "Rendering|Components|Text3D")
-	void SetWordSpacing(const float Value);
+	TEXT3D_API void SetWordSpacing(const float Value);
 
 	/** Get the horizontal alignment value and signal the primitives to be rebuilt */
-	EText3DHorizontalTextAlignment GetHorizontalAlignment() const;
+	TEXT3D_API EText3DHorizontalTextAlignment GetHorizontalAlignment() const;
 
 	/** Set the horizontal alignment value and signal the primitives to be rebuilt */
 	UFUNCTION(BlueprintCallable, Category = "Rendering|Components|Text3D")
-	void SetHorizontalAlignment(const EText3DHorizontalTextAlignment value);
+	TEXT3D_API void SetHorizontalAlignment(const EText3DHorizontalTextAlignment value);
 
 	/** Get the vertical alignment and signal the primitives to be rebuilt */
-	EText3DVerticalTextAlignment GetVerticalAlignment() const;
+	TEXT3D_API EText3DVerticalTextAlignment GetVerticalAlignment() const;
 
 	/** Set the vertical alignment and signal the primitives to be rebuilt */
 	UFUNCTION(BlueprintCallable, Category = "Rendering|Components|Text3D")
-	void SetVerticalAlignment(const EText3DVerticalTextAlignment value);
+	TEXT3D_API void SetVerticalAlignment(const EText3DVerticalTextAlignment value);
 
 	/** Whether a maximum width is specified */
-	bool HasMaxWidth() const;
+	TEXT3D_API bool HasMaxWidth() const;
 
 	/** Enable / Disable a Maximum Width */
 	UFUNCTION(BlueprintCallable, Category = "Rendering|Components|Text3D")
-	void SetHasMaxWidth(const bool Value);
+	TEXT3D_API void SetHasMaxWidth(const bool Value);
 
 	/** Get the Maximum Width - If width is larger, mesh will scale down to fit MaxWidth value */
-	float GetMaxWidth() const;
+	TEXT3D_API float GetMaxWidth() const;
 
 	/** Set the Maximum Width - If width is larger, mesh will scale down to fit MaxWidth value */
 	UFUNCTION(BlueprintCallable, Category = "Rendering|Components|Text3D")
-	void SetMaxWidth(const float Value);
+	TEXT3D_API void SetMaxWidth(const float Value);
 
 	/** Get the Maximum Width Handling - Whether to wrap before scaling when the text size reaches the max width */
-	EText3DMaxWidthHandling GetMaxWidthHandling() const;
+	TEXT3D_API EText3DMaxWidthHandling GetMaxWidthHandling() const;
 
 	/** Set the Maximum Width Handling - Whether to wrap before scaling when the text size reaches the max width */
 	UFUNCTION(BlueprintCallable, Category = "Rendering|Components|Text3D")
-	void SetMaxWidthHandling(const EText3DMaxWidthHandling Value);
+	TEXT3D_API void SetMaxWidthHandling(const EText3DMaxWidthHandling Value);
 
 	/** Whether a maximum height is specified */
-	bool HasMaxHeight() const;
+	TEXT3D_API bool HasMaxHeight() const;
 
 	/** Enable / Disable a Maximum Height */
 	UFUNCTION(BlueprintCallable, Category = "Rendering|Components|Text3D")
-	void SetHasMaxHeight(const bool Value);
+	TEXT3D_API void SetHasMaxHeight(const bool Value);
 
 	/** Get the Maximum Height - If height is larger, mesh will scale down to fit MaxHeight value */
-	float GetMaxHeight() const;
+	TEXT3D_API float GetMaxHeight() const;
 
 	/** Set the Maximum Height - If height is larger, mesh will scale down to fit MaxHeight value */
 	UFUNCTION(BlueprintCallable, Category = "Rendering|Components|Text3D")
-	void SetMaxHeight(const float Value);
+	TEXT3D_API void SetMaxHeight(const float Value);
 
 	/** Get if the mesh should scale proportionally when Max Width/Height is set */
-	bool ScalesProportionally() const;
+	TEXT3D_API bool ScalesProportionally() const;
 
 	/** Set if the mesh should scale proportionally when Max Width/Height is set */
 	UFUNCTION(BlueprintCallable, Category = "Rendering|Components|Text3D")
-	void SetScaleProportionally(const bool Value);
+	TEXT3D_API void SetScaleProportionally(const bool Value);
 
 	/** Get the value of CastShadow. */
-	bool CastsShadow() const;
+	TEXT3D_API bool CastsShadow() const;
 
 	/** Set the value of CastShadow. */
 	UFUNCTION(BlueprintCallable, Category = "Rendering|Components|Text3D")
-	void SetCastShadow(bool NewCastShadow);
+	TEXT3D_API void SetCastShadow(bool NewCastShadow);
 
 	/** Get whole text rendered bounds */
 	UFUNCTION(BlueprintCallable, Category = "Rendering|Components|Text3D")
-	void GetBounds(FVector& Origin, FVector& BoxExtent) const;
+	TEXT3D_API void GetBounds(FVector& Origin, FVector& BoxExtent) const;
 
 	/** Get whole text rendered bounds */
 	FBox GetBounds() const;
@@ -241,7 +245,7 @@ public:
 	}
 
 	/** Set the typeface */
-	void SetTypeface(const FName InTypeface);
+	TEXT3D_API void SetTypeface(const FName InTypeface);
 
 	void SetTextRendererClass(const TSubclassOf<UText3DRendererBase>& InClass);
 	
@@ -256,10 +260,13 @@ public:
 	}
 
 	EText3DRendererFlags GetUpdateFlags() const;
-	void RequestUpdate(EText3DRendererFlags InFlags, bool bInImmediate = false);
+
+	TEXT3D_API void RequestUpdate(EText3DRendererFlags InFlags, bool bInImmediate = false);
 
 	uint32 GetTypeFaceIndex() const;
 
+	const FTypefaceEntry* GetTypeFaceEntry() const;
+	
 	UText3DRendererBase* GetTextRenderer() const;
 
 	UText3DLayoutExtensionBase* GetLayoutExtension() const
@@ -310,23 +317,62 @@ public:
 	{
 		return TransformExtensions;
 	}
+	
+	UText3DCharacterExtensionBase* GetCharacterExtension() const
+	{
+		return CharacterExtension;
+	}
+	
+	template <typename InType UE_REQUIRES(TIsDerivedFrom<InType, UText3DCharacterExtensionBase>::Value)>
+	InType* GetCastedCharacterExtension() const
+	{
+		return Cast<InType>(CharacterExtension);
+	}
+
+	TEXT3D_API uint16 GetCharacterCount() const;
+
+	TEXT3D_API UText3DCharacterBase* GetCharacter(uint16 InCharacterIndex) const;
+
+	template <typename InType UE_REQUIRES(TIsDerivedFrom<InType, UText3DCharacterBase>::Value)>
+	InType* GetCastedCharacter(uint16 InCharacterIndex) const
+	{
+		return Cast<InType>(GetCharacter(InCharacterIndex));
+	}
+
+	TEXT3D_API void ForEachCharacter(const TFunctionRef<void(UText3DCharacterBase*, uint16, uint16)>& InFunctor) const;
+
+	template <typename InType UE_REQUIRES(TIsDerivedFrom<InType, UText3DCharacterBase>::Value)>
+	void ForEachCastedCharacter(const TFunctionRef<void(InType*, uint16, uint16)>& InFunctor) const
+	{
+		ForEachCharacter([&InFunctor](UText3DCharacterBase* InCharacter, uint16 InCharacterIndex, uint16 InCharacterCount)
+		{
+			InFunctor(Cast<InType>(InCharacter), InCharacterIndex, InCharacterCount);
+		});
+	}
 
 protected:
+	//~ Begin USceneComponent
+	/** Intercept and propagate a change on this component to all children. */
+	TEXT3D_API virtual void OnVisibilityChanged() override;
+	/** Intercept and propagate a change on this component to all children. */
+	TEXT3D_API virtual void OnHiddenInGameChanged() override;
+	//~ End USceneComponent
+
 	//~ Begin UActorComponent
-	virtual void OnComponentCreated() override;
-	virtual void OnComponentDestroyed(bool bInDestroyingHierarchy) override;
-	virtual void OnRegister() override;
-	virtual void OnUnregister() override;
+	TEXT3D_API virtual void OnComponentCreated() override;
+	TEXT3D_API virtual void OnComponentDestroyed(bool bInDestroyingHierarchy) override;
+	TEXT3D_API virtual void OnRegister() override;
+	TEXT3D_API virtual void OnUnregister() override;
 	//~ End UActorComponent
 	
 	//~ Begin UObject
-	virtual void Serialize(FArchive& InArchive) override;
-	virtual void PostLoad() override;
-	virtual void PostEditImport() override;
-	virtual void PostDuplicate(EDuplicateMode::Type InDuplicateMode) override;
+	TEXT3D_API virtual void Serialize(FArchive& InArchive) override;
+	TEXT3D_API virtual void PostLoad() override;
+	TEXT3D_API virtual void PostEditImport() override;
+	TEXT3D_API virtual void PostDuplicate(EDuplicateMode::Type InDuplicateMode) override;
 #if WITH_EDITOR
-	virtual void PostEditChangeProperty(FPropertyChangedEvent& PropertyChangedEvent) override;
-	virtual void PostTransacted(const FTransactionObjectEvent& TransactionEvent) override;
+	TEXT3D_API virtual void PostEditChangeProperty(FPropertyChangedEvent& PropertyChangedEvent) override;
+	TEXT3D_API virtual void PostTransacted(const FTransactionObjectEvent& TransactionEvent) override;
 #endif
 	//~ End UObject
 
@@ -334,28 +380,22 @@ protected:
 	friend class UText3DEditorFontSubsystem;
 
 	DECLARE_DELEGATE_RetVal_OneParam(UFont*, FOnResolveFontByName, const FString& InFontName)
-	static FOnResolveFontByName OnResolveFontByNameDelegate;
+	static TEXT3D_API FOnResolveFontByName OnResolveFontByNameDelegate;
 
 	/** Legacy : used to resolve font by their names */
-	void SetFontByName(const FString& InFontName);
+	TEXT3D_API void SetFontByName(const FString& InFontName);
 #endif
 
 	/** Get the type faces supported for the current font */
 	UFUNCTION()
 	TArray<FName> GetTypefaceNames() const;
 
-	/** Intercept and propagate a change on this component to all children. */
-	virtual void OnVisibilityChanged() override;
-
-	/** Intercept and propagate a change on this component to all children. */
-	virtual void OnHiddenInGameChanged() override;
-
 	/**
 	 * Will be called when text geometry is generated.
 	 * Override it to customize text formatting in the final geometry, without affecting the Text property.
 	 * Use GetFormattedText() to retrieve a FText with the result of this formatting.
 	 */
-	virtual void FormatText(FText& InOutText) const;
+	TEXT3D_API virtual void FormatText(FText& InOutText) const;
 
 	/** Called when TextRendererClass is changed */
 	void OnTextRendererClassChanged();
@@ -393,6 +433,9 @@ protected:
 	UPROPERTY(EditAnywhere, Getter = "GetEnforceUpperCase", Setter = "SetEnforceUpperCase", Category = "Text", meta = (AllowPrivateAccess="true"))
 	bool bEnforceUpperCase = false;
 
+	UPROPERTY(VisibleInstanceOnly, Instanced, Category="Character", meta=(AllowPrivateAccess="true"))
+	TObjectPtr<UText3DCharacterExtensionBase> CharacterExtension;
+
 	UPROPERTY(VisibleInstanceOnly, Instanced, Category="Layout", meta=(AllowPrivateAccess="true"))
 	TObjectPtr<UText3DLayoutExtensionBase> LayoutExtension;
 
diff --git a/Engine/Plugins/Experimental/Text3D/Source/Text3D/Public/Text3DTypes.h b/Engine/Plugins/Experimental/Text3D/Source/Text3D/Public/Text3DTypes.h
--- a/Engine/Plugins/Experimental/Text3D/Source/Text3D/Public/Text3DTypes.h
+++ b/Engine/Plugins/Experimental/Text3D/Source/Text3D/Public/Text3DTypes.h
@@ -210,4 +210,30 @@ struct FText3DMaterialGroupKey
 
 	UPROPERTY()
 	EText3DMaterialStyle Style = EText3DMaterialStyle::Invalid;
-};
\ No newline at end of file
+};
+
+struct FText3DCharacterMetadata
+{
+	/** Index of character based on text string (control code, whitespace, ...) */
+	uint16 TextIndex = 0;
+	/** Index of character based on visibility (only renderable character) */
+	uint16 VisibleIndex = 0;
+	/** Index of character based on layout index (include whitespace) */
+	uint16 LayoutIndex = 0;
+	/** Index of character based on glyph index (typeface font specific) */
+	uint16 GlyphIndex = 0;
+	/** Width of the character based on typeface and font metrics */
+	uint16 Width = 0;
+	/** Height of the character based on typeface and font metrics */
+	uint16 Height = 0;
+	/** Actual render bounds */
+	FBox MeshBounds = FBox(ForceInitToZero);
+};
+
+namespace UE::Text3D::Metrics
+{
+	/** Original value to match size across different Text3D version */
+	static constexpr float FontSize = 48;
+	/** Scale used to transform freetype face result to get normalized values */
+	static constexpr float FontSizeInverse = 1.0f / FontSize;
+}
\ No newline at end of file
diff --git a/Engine/Plugins/Experimental/Text3D/Source/Text3D/Public/Utilities/Text3DUtilities.h b/Engine/Plugins/Experimental/Text3D/Source/Text3D/Public/Utilities/Text3DUtilities.h
--- a/Engine/Plugins/Experimental/Text3D/Source/Text3D/Public/Utilities/Text3DUtilities.h
+++ b/Engine/Plugins/Experimental/Text3D/Source/Text3D/Public/Utilities/Text3DUtilities.h
@@ -7,6 +7,7 @@
 
 class UFontFace;
 class UFont;
+struct FTypefaceEntry;
 
 /** Utilities to manipulate text font */
 namespace UE::Text3D::Utilities
@@ -28,4 +29,7 @@ namespace UE::Text3D::Utilities
 	
 	/** Get font faces that compose the font */
 	TEXT3D_API bool GetFontFaces(const UFont* InFont, TArray<UFontFace*>& OutFontFaces);
+
+	/** Get indexes of character for a text based on typeface */
+	TEXT3D_API TArray<FText3DCharacterMetadata> GetCharacterMetadata(const FTypefaceEntry& InTypeFace, const FString& InText);
 }
\ No newline at end of file
diff --git a/Engine/Plugins/Experimental/Text3D/Source/Text3DEditor/Private/Customizations/Text3DEditorHorizontalPropertyTypeCustomization.h b/Engine/Plugins/Experimental/Text3D/Source/Text3DEditor/Private/Customizations/Text3DEditorHorizontalPropertyTypeCustomization.h
--- a/Engine/Plugins/Experimental/Text3D/Source/Text3DEditor/Private/Customizations/Text3DEditorHorizontalPropertyTypeCustomization.h
+++ b/Engine/Plugins/Experimental/Text3D/Source/Text3DEditor/Private/Customizations/Text3DEditorHorizontalPropertyTypeCustomization.h
@@ -3,11 +3,25 @@
 #pragma once
 
 #include "IPropertyTypeCustomization.h"
+#include "PropertyEditorModule.h"
+#include "PropertyHandle.h"
 #include "Templates/SharedPointer.h"
 
 class FDetailWidgetRow;
 class IDetailChildrenBuilder;
 
+/** Only allow property customization with "CustomAlignmentWidget" metadata */
+class FText3DEditorHorizontalPropertyTypeIdentifier : public IPropertyTypeIdentifier
+{
+public:
+	//~ Begin IPropertyTypeIdentifier
+	virtual bool IsPropertyTypeCustomized(const IPropertyHandle& InPropertyHandle) const override
+	{
+		return InPropertyHandle.HasMetaData(TEXT("CustomAlignmentWidget"));
+	}
+	//~ End IPropertyTypeIdentifier
+};
+
 /** Used to customize EText3DHorizontalTextAlignment enum */
 class FText3DEditorHorizontalPropertyTypeCustomization : public IPropertyTypeCustomization
 {
diff --git a/Engine/Plugins/Experimental/Text3D/Source/Text3DEditor/Private/Customizations/Text3DEditorVerticalPropertyTypeCustomization.h b/Engine/Plugins/Experimental/Text3D/Source/Text3DEditor/Private/Customizations/Text3DEditorVerticalPropertyTypeCustomization.h
--- a/Engine/Plugins/Experimental/Text3D/Source/Text3DEditor/Private/Customizations/Text3DEditorVerticalPropertyTypeCustomization.h
+++ b/Engine/Plugins/Experimental/Text3D/Source/Text3DEditor/Private/Customizations/Text3DEditorVerticalPropertyTypeCustomization.h
@@ -3,11 +3,25 @@
 #pragma once
 
 #include "IPropertyTypeCustomization.h"
+#include "PropertyEditorModule.h"
+#include "PropertyHandle.h"
 #include "Templates/SharedPointer.h"
 
 class FDetailWidgetRow;
 class IDetailChildrenBuilder;
 
+/** Only allow property customization with "CustomAlignmentWidget" metadata */
+class FText3DEditorVerticalPropertyTypeIdentifier : public IPropertyTypeIdentifier
+{
+public:
+	//~ Begin IPropertyTypeIdentifier
+	virtual bool IsPropertyTypeCustomized(const IPropertyHandle& InPropertyHandle) const override
+	{
+		return InPropertyHandle.HasMetaData(TEXT("CustomAlignmentWidget"));
+	}
+	//~ End IPropertyTypeIdentifier
+};
+
 /** Used to customize EText3DVerticalTextAlignment enum */
 class FText3DEditorVerticalPropertyTypeCustomization : public IPropertyTypeCustomization
 {
diff --git a/Engine/Plugins/Experimental/Text3D/Source/Text3DEditor/Private/Text3DEditorModule.cpp b/Engine/Plugins/Experimental/Text3D/Source/Text3DEditor/Private/Text3DEditorModule.cpp
--- a/Engine/Plugins/Experimental/Text3D/Source/Text3DEditor/Private/Text3DEditorModule.cpp
+++ b/Engine/Plugins/Experimental/Text3D/Source/Text3DEditor/Private/Text3DEditorModule.cpp
@@ -28,13 +28,16 @@ void FText3DEditorModule::StartupModule()
 	FPropertyEditorModule& PropertyModule = FModuleManager::GetModuleChecked<FPropertyEditorModule>(TEXT("PropertyEditor"));
 
 	PropertyModule.RegisterCustomPropertyTypeLayout(RegisteredTypeNames.Add_GetRef(StaticEnum<EText3DHorizontalTextAlignment>()->GetFName())
-		, FOnGetPropertyTypeCustomizationInstance::CreateStatic(&FText3DEditorHorizontalPropertyTypeCustomization::MakeInstance));
+		, FOnGetPropertyTypeCustomizationInstance::CreateStatic(&FText3DEditorHorizontalPropertyTypeCustomization::MakeInstance)
+		, MakeShared<FText3DEditorHorizontalPropertyTypeIdentifier>());
 
 	PropertyModule.RegisterCustomPropertyTypeLayout(RegisteredTypeNames.Add_GetRef(StaticEnum<EText3DVerticalTextAlignment>()->GetFName())
-		, FOnGetPropertyTypeCustomizationInstance::CreateStatic(&FText3DEditorVerticalPropertyTypeCustomization::MakeInstance));
+		, FOnGetPropertyTypeCustomizationInstance::CreateStatic(&FText3DEditorVerticalPropertyTypeCustomization::MakeInstance)
+		, MakeShared<FText3DEditorVerticalPropertyTypeIdentifier>());
 	
 	PropertyModule.RegisterCustomPropertyTypeLayout(RegisteredTypeNames.Add_GetRef(UFont::StaticClass()->GetFName())
-		, FOnGetPropertyTypeCustomizationInstance::CreateStatic(&FText3DEditorFontPropertyTypeCustomization::MakeInstance), MakeShared<FText3DEditorFontPropertyTypeIdentifier>());
+		, FOnGetPropertyTypeCustomizationInstance::CreateStatic(&FText3DEditorFontPropertyTypeCustomization::MakeInstance)
+		, MakeShared<FText3DEditorFontPropertyTypeIdentifier>());
 }
 
 void FText3DEditorModule::ShutdownModule()
