MotionDesign : Text3D
- Restoring text layout calculation to use same metrics as slate
- Fix avoid trimming spaces on text and causing layout index issues

diff --git a/Engine/Plugins/Experimental/Text3D/Source/Text3D/Private/Extensions/Text3DDefaultGeometryExtension.cpp b/Engine/Plugins/Experimental/Text3D/Source/Text3D/Private/Extensions/Text3DDefaultGeometryExtension.cpp
--- a/Engine/Plugins/Experimental/Text3D/Source/Text3D/Private/Extensions/Text3DDefaultGeometryExtension.cpp
+++ b/Engine/Plugins/Experimental/Text3D/Source/Text3D/Private/Extensions/Text3DDefaultGeometryExtension.cpp
@@ -141,6 +141,7 @@ UStaticMesh* UText3DDefaultGeometryExtension::GetGlyphMesh(UText3DCharacterBase*
 	const FText3DCachedMesh* GlyphMesh = CachedFontData->GetGlyphMesh(GlyphIndex, GlyphMeshParameters);
 	checkf(GlyphMesh != nullptr, TEXT("Invalid glyph mesh returned"));
 	InCharacter->GetCharacterMetadata().MeshBounds = GlyphMesh->MeshBounds;
+	InCharacter->GetCharacterMetadata().MeshOffset = GlyphMesh->MeshOffset;
 	return GlyphMesh->Mesh;
 }
 
diff --git a/Engine/Plugins/Experimental/Text3D/Source/Text3D/Private/Extensions/Text3DDefaultLayoutExtension.cpp b/Engine/Plugins/Experimental/Text3D/Source/Text3D/Private/Extensions/Text3DDefaultLayoutExtension.cpp
--- a/Engine/Plugins/Experimental/Text3D/Source/Text3D/Private/Extensions/Text3DDefaultLayoutExtension.cpp
+++ b/Engine/Plugins/Experimental/Text3D/Source/Text3D/Private/Extensions/Text3DDefaultLayoutExtension.cpp
@@ -18,6 +18,7 @@
 #include "Styling/SlateTypes.h"
 #include "Subsystems/Text3DEngineSubsystem.h"
 #include "Text3DComponent.h"
+#include "Utilities/Text3DUtilities.h"
 
 FName UText3DDefaultLayoutExtension::GetUseMaxWidthPropertyName()
 {
@@ -169,7 +170,7 @@ void UText3DDefaultLayoutExtension::SetScaleProportionally(const bool Value)
 
 float UText3DDefaultLayoutExtension::GetTextHeight() const
 {
-	return ShapedText->Lines.Num() * ShapedText->FontHeight + (ShapedText->Lines.Num() - 1) * LineSpacing;
+	return ShapedText->Lines.Num() * TypeFaceMetrics.FontHeight + (ShapedText->Lines.Num() - 1) * LineSpacing;
 }
 
 FVector UText3DDefaultLayoutExtension::GetTextScale() const
@@ -208,14 +209,29 @@ EText3DExtensionResult UText3DDefaultLayoutExtension::PreRendererUpdate(EText3DR
 
 	if (EnumHasAnyFlags(Text3DComponent->GetUpdateFlags(), EText3DRendererFlags::Geometry))
 	{
-		FSlateFontInfo FontInfo(Text3DComponent->GetFont(), UE::Text3D::Metrics::FontSize);
+		FSlateFontInfo FontInfo(Text3DComponent->GetFont(), UE::Text3D::Metrics::SlateFontSize);
 		FontInfo.CompositeFont = FStyleDefaults::GetFontInfo().CompositeFont;
 		FontInfo.TypefaceFontName = Text3DComponent->GetTypeface();
 		FontStyle->SetFont(FontInfo);
+
+		UText3DEngineSubsystem* Subsystem = UText3DEngineSubsystem::Get();
+
+		const uint32 TypefaceIndex = Text3DComponent->GetTypeFaceIndex();
+		FCachedFontData& CachedFontData = Subsystem->GetCachedFontData(Text3DComponent->GetFont(), TypefaceIndex);
+		const FT_Face Face = CachedFontData.GetFreeTypeFace(TypefaceIndex);
+
+		if (!Face)
+		{
+			UE_LOG(LogText3D, Error, TEXT("Failed to load font face data '%s' in Text3D layout extension"), *CachedFontData.GetFontName());
+			return EText3DExtensionResult::Failed;
+		}
+
+		TypeFaceMetrics.FontHeight = Face->size->metrics.height * UE::Text3D::Metrics::FontSizeInverse;
+		TypeFaceMetrics.FontAscender = Face->size->metrics.ascender * UE::Text3D::Metrics::FontSizeInverse;
+		TypeFaceMetrics.FontDescender = Face->size->metrics.descender * UE::Text3D::Metrics::FontSizeInverse;
 	}
 
 	ShapedText->Reset();
-	ShapedText->FontHeight = FontStyle->Font.Size;
 	FText3DLayoutShaper::Get()->ShapeBidirectionalText(*FontStyle, Text3DComponent->GetFormattedText().ToString(), TextLayout, TextLayoutMarshaller, ShapedText->Lines);
 
 	const int32 ApproximateGlyphNum = Algo::TransformAccumulate(ShapedText->Lines,
@@ -294,10 +310,10 @@ EText3DExtensionResult UText3DDefaultLayoutExtension::PreRendererUpdate(EText3DR
 					Shift = FVector::ZeroVector;
 					break;
 				case EText3DHorizontalTextAlignment::Center:
-					Shift = GlyphSize / 2;
+					Shift.Y = (GlyphSize.Y / 2);
 					break;
 				case EText3DHorizontalTextAlignment::Right:
-					Shift = GlyphSize;
+					Shift.Y = GlyphSize.Y;
 					break;
 				}
 
@@ -306,7 +322,7 @@ EText3DExtensionResult UText3DDefaultLayoutExtension::PreRendererUpdate(EText3DR
 					Shift.Y += ShapedText->Kernings[*CharacterIndex];
 				}
 
-				CharacterTransform.Accumulate(FTransform(Location + FVector(0, Shift.Y, 0)));
+				CharacterTransform.Accumulate(FTransform(Location + Shift));
 
 				for (UText3DTransformExtensionBase* TransformExtension : TransformExtensions)
 				{
@@ -391,7 +407,7 @@ FVector UText3DDefaultLayoutExtension::GetLineLocation(int32 LineIndex) const
 	if (VerticalAlignment != EText3DVerticalTextAlignment::FirstLine)
 	{
 		// First align it to Top
-		VerticalOffset -= ShapedText->FontHeight;
+		VerticalOffset -= TypeFaceMetrics.FontAscender;
 
 		if (VerticalAlignment == EText3DVerticalTextAlignment::Center)
 		{
@@ -399,11 +415,11 @@ FVector UText3DDefaultLayoutExtension::GetLineLocation(int32 LineIndex) const
 		}
 		else if (VerticalAlignment == EText3DVerticalTextAlignment::Bottom)
 		{
-			VerticalOffset += TotalHeight;
+			VerticalOffset += TotalHeight + TypeFaceMetrics.FontDescender;
 		}
 	}
 
-	VerticalOffset -= LineIndex * (ShapedText->FontHeight + LineSpacing);
+	VerticalOffset -= LineIndex * (TypeFaceMetrics.FontHeight + LineSpacing);
 
 	return FVector(0.0f, HorizontalOffset, VerticalOffset);
 }
diff --git a/Engine/Plugins/Experimental/Text3D/Source/Text3D/Private/GeometryBuilders/Text3DGlyphMeshBuilder.cpp b/Engine/Plugins/Experimental/Text3D/Source/Text3D/Private/GeometryBuilders/Text3DGlyphMeshBuilder.cpp
--- a/Engine/Plugins/Experimental/Text3D/Source/Text3D/Private/GeometryBuilders/Text3DGlyphMeshBuilder.cpp
+++ b/Engine/Plugins/Experimental/Text3D/Source/Text3D/Private/GeometryBuilders/Text3DGlyphMeshBuilder.cpp
@@ -95,7 +95,7 @@ void FText3DGlyphMeshBuilder::MirrorGroups(const float Extrude)
 	}
 }
 
-void FText3DGlyphMeshBuilder::MovePivot(const FVector& InNewPivot) const
+void FText3DGlyphMeshBuilder::MovePivot(const FVector& InNewPivot)
 {
 	const TVertexAttributesRef<FVector3f> VertexPositions = Glyph->GetStaticMeshAttributes().GetVertexPositions();
 
@@ -113,17 +113,17 @@ void FText3DGlyphMeshBuilder::MovePivot(const FVector& InNewPivot) const
 		MaxY = FMath::Max(MaxY, Position.Y);
 	}
 
-	// Apply offset with adjusted values
-	const float CenterX = (MinX + MaxX) * 0.5f;
-	FVector3f OffsetVector(-CenterX, -MinY, 0.0f);
+	// Adjust mesh to remove spaces around
+	const double CenterX = (MinX + MaxX) * 0.5f;
+	MeshOffset = FVector(-CenterX, -MinY, 0.0f);
 
+	// Apply pivot offset with adjusted values
 	const FVector3f PivotOffset = FVector3f(InNewPivot) * FVector3f(MaxX - MinX, MaxY - MinY, 0);
-	OffsetVector += PivotOffset;
 
 	for (const FVertexID VertexID : Glyph->GetMeshDescription().Vertices().GetElementIDs())
 	{
 		FVector3f CurrentPosition = VertexPositions[VertexID];
-		CurrentPosition += OffsetVector;
+		CurrentPosition += PivotOffset;
 		VertexPositions[VertexID] = CurrentPosition;
 	}
 }
@@ -138,6 +138,11 @@ FBox FText3DGlyphMeshBuilder::GetMeshBounds() const
 	return Glyph->GetMeshDescription().ComputeBoundingBox();
 }
 
+FVector FText3DGlyphMeshBuilder::GetMeshOffset() const
+{
+	return MeshOffset;
+}
+
 void FText3DGlyphMeshBuilder::CreateFrontMesh(const TText3DGlyphContourNodeShared& Root, const bool bOutline, const float& OutlineExpand)
 {
 	int32 VertexCount = 0;
diff --git a/Engine/Plugins/Experimental/Text3D/Source/Text3D/Private/GeometryBuilders/Text3DGlyphMeshBuilder.h b/Engine/Plugins/Experimental/Text3D/Source/Text3D/Private/GeometryBuilders/Text3DGlyphMeshBuilder.h
--- a/Engine/Plugins/Experimental/Text3D/Source/Text3D/Private/GeometryBuilders/Text3DGlyphMeshBuilder.h
+++ b/Engine/Plugins/Experimental/Text3D/Source/Text3D/Private/GeometryBuilders/Text3DGlyphMeshBuilder.h
@@ -33,19 +33,22 @@ public:
 	void MirrorGroups(const float Extrude);
 
 	/** Move the mesh pivot based on pivot ratio 0-1 */
-	void MovePivot(const FVector& InNewPivot) const;
+	void MovePivot(const FVector& InNewPivot);
 
 	void BuildMesh(UStaticMesh* StaticMesh, UMaterial* DefaultMaterial);
 
 	/** Get actual glyph bounds once built */
 	FBox GetMeshBounds() const;
 
+	/** Get actual glyph offset */
+	FVector GetMeshOffset() const;
+
 private:
+	FVector MeshOffset = FVector::ZeroVector;
 	TSharedPtr<FText3DGlyph> Glyph;
 	TSharedRef<FText3DGlyphData> Data;
 	TSharedPtr<FText3DGlyphContourList> Contours;
 
-
 	/**
 	 * Create 'Front' part of glyph.
 	 * @param Root - Tree root.
diff --git a/Engine/Plugins/Experimental/Text3D/Source/Text3D/Private/LayoutBuilders/Text3DShapedGlyphText.cpp b/Engine/Plugins/Experimental/Text3D/Source/Text3D/Private/LayoutBuilders/Text3DShapedGlyphText.cpp
--- a/Engine/Plugins/Experimental/Text3D/Source/Text3D/Private/LayoutBuilders/Text3DShapedGlyphText.cpp
+++ b/Engine/Plugins/Experimental/Text3D/Source/Text3D/Private/LayoutBuilders/Text3DShapedGlyphText.cpp
@@ -11,7 +11,6 @@ FText3DShapedGlyphText::FText3DShapedGlyphText()
 
 void FText3DShapedGlyphText::Reset()
 {
-	FontHeight = 0.0f;
 	PivotHAlignment = EText3DHorizontalTextAlignment::Left;
 	PivotVAlignment = EText3DVerticalTextAlignment::Bottom;
 	Kernings.Reset();
@@ -41,13 +40,6 @@ void FText3DShapedGlyphText::CalculateWidth()
 
 		for (TConstEnumerateRef<FShapedGlyphEntry> CurrentGlyph : EnumerateRange(GlyphLine->GlyphsToRender))
 		{
-			// Trim for proper positioning
-			if (!CurrentGlyph->bIsVisible && CurrentGlyph.GetIndex() == 0)
-			{
-				LayoutIndex++;
-				continue;
-			}
-
 			const int32* CharacterIndex = CharacterGlyphMap.Find(LayoutIndex);
 			const float CurrentKerning = !!CharacterIndex ? Kernings[*CharacterIndex] : Kerning;
 			const float GlyphAdv = GlyphLine->GetWidthAdvance(CurrentGlyph.GetIndex(), CurrentKerning, WordSpacing);
diff --git a/Engine/Plugins/Experimental/Text3D/Source/Text3D/Private/LayoutBuilders/Text3DShapedGlyphText.h b/Engine/Plugins/Experimental/Text3D/Source/Text3D/Private/LayoutBuilders/Text3DShapedGlyphText.h
--- a/Engine/Plugins/Experimental/Text3D/Source/Text3D/Private/LayoutBuilders/Text3DShapedGlyphText.h
+++ b/Engine/Plugins/Experimental/Text3D/Source/Text3D/Private/LayoutBuilders/Text3DShapedGlyphText.h
@@ -19,9 +19,6 @@ struct FText3DShapedGlyphText
 	/** Maps character layout index to character index */
 	TMap<int32, int32> CharacterGlyphMap;
 
-	/** Font height expected */
-	float FontHeight;
-
 	/** Pivot Horizontal alignment */
 	EText3DHorizontalTextAlignment PivotHAlignment;
 
diff --git a/Engine/Plugins/Experimental/Text3D/Source/Text3D/Private/Subsystems/Text3DEngineSubsystem.cpp b/Engine/Plugins/Experimental/Text3D/Source/Text3D/Private/Subsystems/Text3DEngineSubsystem.cpp
--- a/Engine/Plugins/Experimental/Text3D/Source/Text3D/Private/Subsystems/Text3DEngineSubsystem.cpp
+++ b/Engine/Plugins/Experimental/Text3D/Source/Text3D/Private/Subsystems/Text3DEngineSubsystem.cpp
@@ -250,7 +250,8 @@ void FCachedFontData::LoadFreeTypeFace(uint32 InTypefaceIndex)
 	if (const FT_Face* FreeTypeface = &TypefaceFontDataMap[InTypefaceIndex].GetTypeface())
 	{
 		TypefaceFontDataMap[InTypefaceIndex].SetTypefaceFontDataHash(HashCombine(0, GetTypeHash(Typeface.Font)));
-		FT_Set_Pixel_Sizes(*FreeTypeface, 0, UE::Text3D::Metrics::FontSize);
+		FT_Set_Char_Size(*FreeTypeface, UE::Text3D::Metrics::FontSize, UE::Text3D::Metrics::FontSize, UE::Text3D::Metrics::FontDPI, UE::Text3D::Metrics::FontDPI);
+		FT_Set_Pixel_Sizes(*FreeTypeface, UE::Text3D::Metrics::FontSize, UE::Text3D::Metrics::FontSize);
 	}
 }
 
@@ -331,6 +332,7 @@ FText3DCachedMesh* FCachedFontData::GetGlyphMesh(uint32 InGlyphIndex, const FGly
 	MeshCreator.BuildMesh(CachedMesh.Mesh, Subsystem->DefaultMaterial);
 
 	CachedMesh.MeshBounds = MeshCreator.GetMeshBounds();
+	CachedMesh.MeshOffset = MeshCreator.GetMeshOffset();
 
 	return &CachedMesh;
 }
diff --git a/Engine/Plugins/Experimental/Text3D/Source/Text3D/Private/Subsystems/Text3DEngineSubsystem.h b/Engine/Plugins/Experimental/Text3D/Source/Text3D/Private/Subsystems/Text3DEngineSubsystem.h
--- a/Engine/Plugins/Experimental/Text3D/Source/Text3D/Private/Subsystems/Text3DEngineSubsystem.h
+++ b/Engine/Plugins/Experimental/Text3D/Source/Text3D/Private/Subsystems/Text3DEngineSubsystem.h
@@ -69,7 +69,10 @@ struct FText3DCachedMesh
 	TObjectPtr<UStaticMesh> Mesh;
 
 	UPROPERTY()
-	FBox MeshBounds;
+	FBox MeshBounds = FBox(ForceInitToZero);
+
+	UPROPERTY()
+	FVector MeshOffset = FVector::ZeroVector;
 };
 
 USTRUCT()
diff --git a/Engine/Plugins/Experimental/Text3D/Source/Text3D/Private/Utilities/Text3DUtilities.cpp b/Engine/Plugins/Experimental/Text3D/Source/Text3D/Private/Utilities/Text3DUtilities.cpp
--- a/Engine/Plugins/Experimental/Text3D/Source/Text3D/Private/Utilities/Text3DUtilities.cpp
+++ b/Engine/Plugins/Experimental/Text3D/Source/Text3D/Private/Utilities/Text3DUtilities.cpp
@@ -279,3 +279,26 @@ TArray<FText3DCharacterMetadata> UE::Text3D::Utilities::GetCharacterMetadata(con
 
 	return CharacterIndexes;
 }
+
+FText3DTypeFaceMetrics UE::Text3D::Utilities::GetTypeFaceMetrics(const FTypefaceEntry& InTypeFace)
+{
+	FText3DTypeFaceMetrics Metrics;
+
+	const FFontFaceDataConstPtr FaceData = InTypeFace.Font.GetFontFaceData();
+
+	if (FaceData.IsValid() && FaceData->HasData() && !FaceData->GetData().IsEmpty())
+	{
+		const TArray<uint8> Data = FaceData->GetData();
+
+		FT_Face FreeTypeFace;
+		FT_New_Memory_Face(FText3DModule::GetFreeTypeLibrary(), Data.GetData(), Data.Num(), 0, &FreeTypeFace);
+
+		Metrics.FontAscender = FreeTypeFace->ascender * UE::Text3D::Metrics::FontSizeInverse;
+		Metrics.FontDescender = FreeTypeFace->descender * UE::Text3D::Metrics::FontSizeInverse;
+		Metrics.FontHeight = FreeTypeFace->height * UE::Text3D::Metrics::FontSizeInverse;
+
+		FT_Done_Face(FreeTypeFace);
+	}
+
+	return Metrics;
+}
diff --git a/Engine/Plugins/Experimental/Text3D/Source/Text3D/Public/Extensions/Text3DDefaultLayoutExtension.h b/Engine/Plugins/Experimental/Text3D/Source/Text3D/Public/Extensions/Text3DDefaultLayoutExtension.h
--- a/Engine/Plugins/Experimental/Text3D/Source/Text3D/Public/Extensions/Text3DDefaultLayoutExtension.h
+++ b/Engine/Plugins/Experimental/Text3D/Source/Text3D/Public/Extensions/Text3DDefaultLayoutExtension.h
@@ -7,8 +7,9 @@
 #include "Text3DDefaultLayoutExtension.generated.h"
 
 class FText3DLayout;
-struct FTextBlockStyle;
 struct FPropertyChangedEvent;
+struct FTextBlockStyle;
+struct FText3DTypeFaceMetrics;
 
 UCLASS(MinimalAPI)
 class UText3DDefaultLayoutExtension : public UText3DLayoutExtensionBase
@@ -222,4 +223,6 @@ private:
 
 	/** Determines how text is laid out, ie. parsing line breaks. */
 	TSharedPtr<ITextLayoutMarshaller> TextLayoutMarshaller;
+
+	FText3DTypeFaceMetrics TypeFaceMetrics;
 };
\ No newline at end of file
diff --git a/Engine/Plugins/Experimental/Text3D/Source/Text3D/Public/Text3DTypes.h b/Engine/Plugins/Experimental/Text3D/Source/Text3D/Public/Text3DTypes.h
--- a/Engine/Plugins/Experimental/Text3D/Source/Text3D/Public/Text3DTypes.h
+++ b/Engine/Plugins/Experimental/Text3D/Source/Text3D/Public/Text3DTypes.h
@@ -228,12 +228,27 @@ struct FText3DCharacterMetadata
 	uint16 Height = 0;
 	/** Actual render bounds */
 	FBox MeshBounds = FBox(ForceInitToZero);
+	/** Offset around glyph due to font face */
+	FVector MeshOffset = FVector::ZeroVector;
+};
+
+struct FText3DTypeFaceMetrics
+{
+	float FontHeight = 0.f;
+	/** Height above baseline */
+	float FontAscender = 0.f;
+	/** Height below baseline */
+	float FontDescender = 0.f;
 };
 
 namespace UE::Text3D::Metrics
 {
 	/** Original value to match size across different Text3D version */
-	static constexpr float FontSize = 48;
+	static constexpr float FontSize = 64;
+	/** Value used to match size across Text3D implementation */
+	static constexpr float SlateFontSize = 48;
 	/** Scale used to transform freetype face result to get normalized values */
 	static constexpr float FontSizeInverse = 1.0f / FontSize;
+	/** DPI used to match slate */
+	static constexpr float FontDPI = 96.f;
 }
\ No newline at end of file
diff --git a/Engine/Plugins/Experimental/Text3D/Source/Text3D/Public/Utilities/Text3DUtilities.h b/Engine/Plugins/Experimental/Text3D/Source/Text3D/Public/Utilities/Text3DUtilities.h
--- a/Engine/Plugins/Experimental/Text3D/Source/Text3D/Public/Utilities/Text3DUtilities.h
+++ b/Engine/Plugins/Experimental/Text3D/Source/Text3D/Public/Utilities/Text3DUtilities.h
@@ -7,6 +7,7 @@
 
 class UFontFace;
 class UFont;
+struct FText3DTypeFaceMetrics;
 struct FTypefaceEntry;
 
 /** Utilities to manipulate text font */
@@ -32,4 +33,7 @@ namespace UE::Text3D::Utilities
 
 	/** Get indexes of character for a text based on typeface */
 	TEXT3D_API TArray<FText3DCharacterMetadata> GetCharacterMetadata(const FTypefaceEntry& InTypeFace, const FString& InText);
+
+	/** Get type face font metrics */
+	TEXT3D_API FText3DTypeFaceMetrics GetTypeFaceMetrics(const FTypefaceEntry& InTypeFace);
 }
\ No newline at end of file
