MotionDesign : Text3D
- Adding support for additional rendering options like affect indirect lighting, holdout in Text3D rendering extension

diff --git a/Engine/Plugins/Experimental/Text3D/Source/Text3D/Private/Extensions/Text3DDefaultRenderingExtension.cpp b/Engine/Plugins/Experimental/Text3D/Source/Text3D/Private/Extensions/Text3DDefaultRenderingExtension.cpp
--- a/Engine/Plugins/Experimental/Text3D/Source/Text3D/Private/Extensions/Text3DDefaultRenderingExtension.cpp
+++ b/Engine/Plugins/Experimental/Text3D/Source/Text3D/Private/Extensions/Text3DDefaultRenderingExtension.cpp
@@ -24,6 +24,39 @@ void UText3DDefaultRenderingExtension::SetCastHiddenShadow(bool bInCastShadow)
 	OnRenderingOptionsChanged();
 }
 
+void UText3DDefaultRenderingExtension::SetAffectDynamicIndirectLighting(bool bInValue)
+{
+	if (bAffectDynamicIndirectLighting == bInValue)
+	{
+		return;
+	}
+
+	bAffectDynamicIndirectLighting = bInValue;
+	OnRenderingOptionsChanged();
+}
+
+void UText3DDefaultRenderingExtension::SetAffectIndirectLightingWhileHidden(bool bInValue)
+{
+	if (bAffectIndirectLightingWhileHidden == bInValue)
+	{
+		return;
+	}
+
+	bAffectIndirectLightingWhileHidden = bInValue;
+	OnRenderingOptionsChanged();
+}
+
+void UText3DDefaultRenderingExtension::SetHoldout(bool bInHoldout)
+{
+	if (bHoldout == bInHoldout)
+	{
+		return;
+	}
+
+	bHoldout = bInHoldout;
+	OnRenderingOptionsChanged();
+}
+
 #if WITH_EDITOR
 void UText3DDefaultRenderingExtension::PostEditChangeProperty(FPropertyChangedEvent& InEvent)
 {
@@ -32,7 +65,10 @@ void UText3DDefaultRenderingExtension::PostEditChangeProperty(FPropertyChangedEv
 	static const TSet<FName> PropertyNames =
 	{
 		GET_MEMBER_NAME_CHECKED(UText3DDefaultRenderingExtension, bCastShadow),
-		GET_MEMBER_NAME_CHECKED(UText3DDefaultRenderingExtension, bCastHiddenShadow)
+		GET_MEMBER_NAME_CHECKED(UText3DDefaultRenderingExtension, bCastHiddenShadow),
+		GET_MEMBER_NAME_CHECKED(UText3DDefaultRenderingExtension, bAffectDynamicIndirectLighting),
+		GET_MEMBER_NAME_CHECKED(UText3DDefaultRenderingExtension, bAffectIndirectLightingWhileHidden),
+		GET_MEMBER_NAME_CHECKED(UText3DDefaultRenderingExtension, bHoldout)
 	};
 
 	if (PropertyNames.Contains(InEvent.GetMemberPropertyName()))
@@ -47,19 +83,24 @@ bool UText3DDefaultRenderingExtension::GetTextCastShadow() const
 	return GetCastShadow();
 }
 
-void UText3DDefaultRenderingExtension::SetTextCastShadow(bool bInValue)
+bool UText3DDefaultRenderingExtension::GetTextCastHiddenShadow() const
 {
-	SetCastShadow(bInValue);
+	return GetCastHiddenShadow();
 }
 
-bool UText3DDefaultRenderingExtension::GetTextCastHiddenShadow() const
+bool UText3DDefaultRenderingExtension::GetTextAffectDynamicIndirectLighting() const
 {
-	return GetCastHiddenShadow();
+	return GetAffectDynamicIndirectLighting();
+}
+
+bool UText3DDefaultRenderingExtension::GetTextAffectIndirectLightingWhileHidden() const
+{
+	return GetAffectIndirectLightingWhileHidden();
 }
 
-void UText3DDefaultRenderingExtension::SetTextCastHiddenShadow(bool bInValue)
+bool UText3DDefaultRenderingExtension::GetTextHoldout() const
 {
-	SetCastHiddenShadow(bInValue);
+	return GetHoldout();
 }
 
 void UText3DDefaultRenderingExtension::OnRenderingOptionsChanged()
diff --git a/Engine/Plugins/Experimental/Text3D/Source/Text3D/Private/Renderers/StaticMeshes/Text3DStaticMeshesRenderer.cpp b/Engine/Plugins/Experimental/Text3D/Source/Text3D/Private/Renderers/StaticMeshes/Text3DStaticMeshesRenderer.cpp
--- a/Engine/Plugins/Experimental/Text3D/Source/Text3D/Private/Renderers/StaticMeshes/Text3DStaticMeshesRenderer.cpp
+++ b/Engine/Plugins/Experimental/Text3D/Source/Text3D/Private/Renderers/StaticMeshes/Text3DStaticMeshesRenderer.cpp
@@ -59,6 +59,7 @@ void UText3DStaticMeshesRenderer::OnUpdate(EText3DRendererFlags InFlags)
 		{
 			constexpr bool bReset = false;
 			const FTransform& CharacterTransform = InCharacter->GetTransform(bReset);
+			CharacterKernings[InIndex]->SetRelativeTransform(FTransform::Identity);
 			CharacterMeshes[InIndex]->SetRelativeTransform(CharacterTransform);
 		});
 
@@ -102,6 +103,9 @@ void UText3DStaticMeshesRenderer::OnUpdate(EText3DRendererFlags InFlags)
 			StaticMeshComponent->SetVisibility(TextComponent->GetVisibleFlag());
 			StaticMeshComponent->SetCastShadow(RenderingExtension->GetTextCastShadow());
 			StaticMeshComponent->SetCastHiddenShadow(RenderingExtension->GetTextCastHiddenShadow());
+			StaticMeshComponent->SetAffectDynamicIndirectLighting(RenderingExtension->GetTextAffectDynamicIndirectLighting());
+			StaticMeshComponent->SetAffectIndirectLightingWhileHidden(RenderingExtension->GetTextAffectIndirectLightingWhileHidden());
+			StaticMeshComponent->SetHoldout(RenderingExtension->GetTextHoldout());
 		}
 	}
 }
@@ -181,8 +185,6 @@ bool UText3DStaticMeshesRenderer::AllocateCharacterComponents(int32 Num)
 		return false;
 	}
 
-	UText3DComponent* TextComponent = GetText3DComponent();
-
 	// Add characters
 	if (FMath::Sign(DeltaNum) > 0)
 	{
@@ -193,20 +195,15 @@ bool UText3DStaticMeshesRenderer::AllocateCharacterComponents(int32 Num)
 
 			const FName CharacterKerningComponentName = MakeUniqueObjectName(this, USceneComponent::StaticClass(), FName(*FString::Printf(TEXT("CharacterKerning%d"), GlyphId)));
 			USceneComponent* CharacterKerningComponent = NewObject<USceneComponent>(this, CharacterKerningComponentName, RF_Transient);
-
 			CharacterKerningComponent->AttachToComponent(TextRoot, FAttachmentTransformRules::KeepRelativeTransform);
 			CharacterKerningComponent->RegisterComponent();
 			CharacterKernings.Add(CharacterKerningComponent);
 
 			const FName StaticMeshComponentName = MakeUniqueObjectName(this, UStaticMeshComponent::StaticClass(), FName(*FString::Printf(TEXT("StaticMeshComponent%d"), GlyphId)));
 			UStaticMeshComponent* StaticMeshComponent = NewObject<UStaticMeshComponent>(this, StaticMeshComponentName, RF_Transient);
+			StaticMeshComponent->AttachToComponent(CharacterKerningComponent, FAttachmentTransformRules::KeepRelativeTransform);
 			StaticMeshComponent->RegisterComponent();
-			StaticMeshComponent->SetVisibility(TextComponent->GetVisibleFlag());
-			StaticMeshComponent->SetHiddenInGame(TextComponent->bHiddenInGame);
-			StaticMeshComponent->SetCastShadow(TextComponent->CastsShadow());
 			CharacterMeshes.Add(StaticMeshComponent);
-
-			StaticMeshComponent->AttachToComponent(CharacterKerningComponent, FAttachmentTransformRules::KeepRelativeTransform);
 		}
 	}
 	// Remove characters
diff --git a/Engine/Plugins/Experimental/Text3D/Source/Text3D/Private/Renderers/Text3DRendererBase.cpp b/Engine/Plugins/Experimental/Text3D/Source/Text3D/Private/Renderers/Text3DRendererBase.cpp
--- a/Engine/Plugins/Experimental/Text3D/Source/Text3D/Private/Renderers/Text3DRendererBase.cpp
+++ b/Engine/Plugins/Experimental/Text3D/Source/Text3D/Private/Renderers/Text3DRendererBase.cpp
@@ -125,7 +125,7 @@ void UText3DRendererBase::Update(EText3DRendererFlags InFlags)
 		Flag <<= 1;
 	}
 
-	UE_LOG(LogText3D, Log, TEXT("%s : Text3DRenderer %s Updated with flags %i"), *Text3DComponent->GetOwner()->GetActorNameOrLabel(), *GetName().ToString(), InFlags)
+	UE_LOG(LogText3D, Verbose, TEXT("%s : Text3DRenderer %s Updated with flags %i"), *Text3DComponent->GetOwner()->GetActorNameOrLabel(), *GetName().ToString(), InFlags)
 }
 
 void UText3DRendererBase::Clear()
diff --git a/Engine/Plugins/Experimental/Text3D/Source/Text3D/Private/Text3DComponent.cpp b/Engine/Plugins/Experimental/Text3D/Source/Text3D/Private/Text3DComponent.cpp
--- a/Engine/Plugins/Experimental/Text3D/Source/Text3D/Private/Text3DComponent.cpp
+++ b/Engine/Plugins/Experimental/Text3D/Source/Text3D/Private/Text3DComponent.cpp
@@ -305,9 +305,9 @@ void UText3DComponent::Serialize(FArchive& InArchive)
 				DefaultLayoutExtension->SetScaleProportionally(bScaleProportionally_DEPRECATED);
 			}
 			
-			if (RenderingExtension)
+			if (UText3DDefaultRenderingExtension* DefaultRenderingExtension = GetCastedRenderingExtension<UText3DDefaultRenderingExtension>())
 			{
-				RenderingExtension->SetTextCastShadow(bCastShadow_DEPRECATED);
+				DefaultRenderingExtension->SetCastShadow(bCastShadow_DEPRECATED);
 			}
 
 			PRAGMA_ENABLE_DEPRECATION_WARNINGS
@@ -783,9 +783,9 @@ bool UText3DComponent::CastsShadow() const
 
 void UText3DComponent::SetCastShadow(bool NewCastShadow)
 {
-	if (RenderingExtension)
+	if (UText3DDefaultRenderingExtension* DefaultRenderingExtension = GetCastedRenderingExtension<UText3DDefaultRenderingExtension>())
 	{
-		RenderingExtension->SetTextCastShadow(NewCastShadow);
+		DefaultRenderingExtension->SetCastShadow(NewCastShadow);
 	}
 }
 
@@ -814,16 +814,21 @@ void UText3DComponent::SetTextRendererClass(const TSubclassOf<UText3DRendererBas
 
 void UText3DComponent::RebuildInternal(bool bCleanCache)
 {
+	if (IsTemplate(RF_ClassDefaultObject))
+	{
+		return;
+	}
+
 	if (!IsRegistered())
 	{
 		return;
 	}
-	
+
 	if (UE::IsSavingPackage(this))
 	{
 		return;
 	}
-	
+
 	if (!TextRenderer)
 	{
 		UE_LOG(LogText3D, Error, TEXT("Text3D renderer is not valid, cannot proceed"));
@@ -941,7 +946,8 @@ void UText3DComponent::OnTextRendererClassChanged()
 
 void UText3DComponent::OnTextChanged()
 {
-	RequestUpdate(EText3DRendererFlags::All);
+	constexpr bool bImmediate = true;
+	RequestUpdate(EText3DRendererFlags::All, bImmediate);
 }
 
 void UText3DComponent::OnFontPropertiesChanged()
diff --git a/Engine/Plugins/Experimental/Text3D/Source/Text3D/Public/Extensions/Text3DDefaultRenderingExtension.h b/Engine/Plugins/Experimental/Text3D/Source/Text3D/Public/Extensions/Text3DDefaultRenderingExtension.h
--- a/Engine/Plugins/Experimental/Text3D/Source/Text3D/Public/Extensions/Text3DDefaultRenderingExtension.h
+++ b/Engine/Plugins/Experimental/Text3D/Source/Text3D/Public/Extensions/Text3DDefaultRenderingExtension.h
@@ -35,6 +35,39 @@ public:
 	UFUNCTION(BlueprintCallable, Category = "Text3D|Rendering")
 	void SetCastHiddenShadow(bool bInCastShadow);
 	
+	/** Get the value of AffectDynamicIndirectLighting. */
+	UFUNCTION(BlueprintPure, Category = "Text3D|Rendering")
+	bool GetAffectDynamicIndirectLighting() const
+	{
+		return bAffectDynamicIndirectLighting;
+	}
+
+	/** Set the value of AffectDynamicIndirectLighting. */
+	UFUNCTION(BlueprintCallable, Category = "Text3D|Rendering")
+	void SetAffectDynamicIndirectLighting(bool bInValue);
+
+	/** Get the value of AffectIndirectLightingWhileHidden. */
+	UFUNCTION(BlueprintPure, Category = "Text3D|Rendering")
+	bool GetAffectIndirectLightingWhileHidden() const
+	{
+		return bAffectIndirectLightingWhileHidden;
+	}
+
+	/** Set the value of AffectIndirectLightingWhileHidden. */
+	UFUNCTION(BlueprintCallable, Category = "Text3D|Rendering")
+	void SetAffectIndirectLightingWhileHidden(bool bInValue);
+
+	/** Get the value of Holdout. */
+	UFUNCTION(BlueprintPure, Category = "Text3D|Rendering")
+	bool GetHoldout() const
+	{
+		return bHoldout;
+	}
+
+	/** Set the value of Holdout. */
+	UFUNCTION(BlueprintCallable, Category = "Text3D|Rendering")
+	void SetHoldout(bool bInHoldout);
+	
 protected:
 	//~ Begin UObject
 #if WITH_EDITOR
@@ -44,9 +77,10 @@ protected:
 
 	//~ Begin UText3DRenderingExtensionBase
 	virtual bool GetTextCastShadow() const override;
-	virtual void SetTextCastShadow(bool bInValue) override;
 	virtual bool GetTextCastHiddenShadow() const override;
-	virtual void SetTextCastHiddenShadow(bool bInValue) override;
+	virtual bool GetTextAffectDynamicIndirectLighting() const override;
+	virtual bool GetTextAffectIndirectLightingWhileHidden() const override;
+	virtual bool GetTextHoldout() const override;
 	//~ End UText3DRenderingExtensionBase
 
 	void OnRenderingOptionsChanged();
@@ -56,6 +90,18 @@ protected:
 	bool bCastShadow = true;
 
 	/** Controls whether the text glyphs should cast a shadow or not. */
-	UPROPERTY(EditAnywhere, Getter = "GetCastShadow", Setter = "SetCastShadow", Category = "Rendering", meta = (AllowPrivateAccess = "true"))
+	UPROPERTY(EditAnywhere, DisplayName = "Hidden Shadow", Getter = "GetCastShadow", Setter = "SetCastShadow", Category = "Rendering", meta = (AllowPrivateAccess = "true"))
 	bool bCastHiddenShadow = false;
+
+	/** Controls whether the primitive should influence indirect lighting. */
+	UPROPERTY(EditAnywhere, Getter = "GetAffectDynamicIndirectLighting", Setter = "SetAffectDynamicIndirectLighting", Category = "Rendering", meta = (AllowPrivateAccess = "true"))
+	bool bAffectDynamicIndirectLighting = false;
+
+	/** Controls whether the primitive should affect indirect lighting when hidden. This flag is only used if bAffectDynamicIndirectLighting is true. */
+	UPROPERTY(EditAnywhere, Getter="GetAffectIndirectLightingWhileHidden", Setter="SetAffectIndirectLightingWhileHidden", Category = "Rendering", meta = (EditCondition = "bAffectDynamicIndirectLighting", AllowPrivateAccess = "true"))
+	bool bAffectIndirectLightingWhileHidden = false;
+	
+	/** If this is True, this primitive will render black with an alpha of 0, but all secondary effects (shadows, reflections, indirect lighting) remain. This feature requires activating the project setting(s) "Alpha Output", and "Support Primitive Alpha Holdout" if using the deferred renderer. */
+	UPROPERTY(EditAnywhere, Getter="GetHoldout", Setter="SetHoldout", Category = "Rendering", meta = (AllowPrivateAccess = "true"))
+	bool bHoldout = false;
 };
diff --git a/Engine/Plugins/Experimental/Text3D/Source/Text3D/Public/Extensions/Text3DRenderingExtensionBase.h b/Engine/Plugins/Experimental/Text3D/Source/Text3D/Public/Extensions/Text3DRenderingExtensionBase.h
--- a/Engine/Plugins/Experimental/Text3D/Source/Text3D/Public/Extensions/Text3DRenderingExtensionBase.h
+++ b/Engine/Plugins/Experimental/Text3D/Source/Text3D/Public/Extensions/Text3DRenderingExtensionBase.h
@@ -17,20 +17,27 @@ public:
 		: UText3DExtensionBase(/** Priority */4)
 	{}
 
-	virtual void SetTextCastShadow(bool bInValue)
+	virtual bool GetTextCastShadow() const
 	{
+		return true;
 	}
 
-	virtual bool GetTextCastShadow() const
+	virtual bool GetTextCastHiddenShadow() const
 	{
-		return true;
+		return false;
 	}
 
-	virtual void SetTextCastHiddenShadow(bool bInValue)
+	virtual bool GetTextAffectDynamicIndirectLighting() const
 	{
+		return false;
 	}
 
-	virtual bool GetTextCastHiddenShadow() const
+	virtual bool GetTextAffectIndirectLightingWhileHidden() const
+	{
+		return false;
+	}
+	
+	virtual bool GetTextHoldout() const
 	{
 		return false;
 	}
