MotionDesign : Text3D
- Reset character state when they are stored in the pool and recycled
- Adding text tracking and character kerning
- Adding per character visibility property

diff --git a/Engine/Plugins/Experimental/Avalanche/Source/AvalancheTextEditor/Private/Visualizer/AvaTextVisualizer.cpp b/Engine/Plugins/Experimental/Avalanche/Source/AvalancheTextEditor/Private/Visualizer/AvaTextVisualizer.cpp
--- a/Engine/Plugins/Experimental/Avalanche/Source/AvalancheTextEditor/Private/Visualizer/AvaTextVisualizer.cpp
+++ b/Engine/Plugins/Experimental/Avalanche/Source/AvalancheTextEditor/Private/Visualizer/AvaTextVisualizer.cpp
@@ -4,7 +4,6 @@
 
 #include "AvaField.h"
 #include "AvaShapeSprites.h"
-#include "AvaText3DComponent.h"
 #include "Characters/Text3DDefaultCharacter.h"
 #include "EditorViewportClient.h"
 #include "Engine/Texture2D.h"
@@ -704,7 +703,6 @@ bool FAvaTextVisualizer::HandleInputDeltaInternal(FEditorViewportClient* InViewp
 							bHasBeenModified = true;
 						}
 
-						Character->SetKerningOverride(true);
 						Character->SetKerning(Kerning);
 						return true;
 					}
@@ -791,7 +789,6 @@ bool FAvaTextVisualizer::ResetValue(FEditorViewportClient* InViewportClient, HHi
 		{
 			Character->Modify();
 			const UText3DDefaultCharacter* CDO = GetDefault<UText3DDefaultCharacter>();
-			Character->SetKerningOverride(CDO->GetKerningOverride());
 			Character->SetKerning(CDO->GetKerning());
 			FAvaVisualizerBase::NotifyPropertyModified(Character, CharacterKerningProperty, EPropertyChangeType::ValueSet);
 		}
@@ -1143,7 +1140,7 @@ void FAvaTextVisualizer::StoreInitialValues()
 	{
 		if (const UText3DCharacterBase* Character = TextComponent->GetCharacter(EditingKerningIndex))
 		{
-			InitialCharacterKerning = Character->GetCharacterKerning().Get(0.f);
+			InitialCharacterKerning = Character->GetCharacterKerning();
 		}
 	}
 }
diff --git a/Engine/Plugins/Experimental/Text3D/Source/Text3D/Private/Characters/Text3DCharacterBase.cpp b/Engine/Plugins/Experimental/Text3D/Source/Text3D/Private/Characters/Text3DCharacterBase.cpp
--- a/Engine/Plugins/Experimental/Text3D/Source/Text3D/Private/Characters/Text3DCharacterBase.cpp
+++ b/Engine/Plugins/Experimental/Text3D/Source/Text3D/Private/Characters/Text3DCharacterBase.cpp
@@ -19,6 +19,11 @@ FName UText3DCharacterBase::GetRelativeScalePropertyName()
 	return GET_MEMBER_NAME_CHECKED(UText3DCharacterBase, RelativeScale);
 }
 
+FName UText3DCharacterBase::GetVisiblePropertyName()
+{
+	return GET_MEMBER_NAME_CHECKED(UText3DCharacterBase, bVisible);
+}
+
 FTransform& UText3DCharacterBase::GetTransform(bool bInReset)
 {
 	if (bInReset)
@@ -42,7 +47,7 @@ void UText3DCharacterBase::SetRelativeLocation(const FVector& InLocation)
 	}
 
 	RelativeLocation = InLocation;
-	OnCharacterDataChanged();
+	OnCharacterDataChanged(EText3DRendererFlags::Layout);
 }
 
 void UText3DCharacterBase::SetRelativeRotation(const FRotator& InRotation)
@@ -53,7 +58,7 @@ void UText3DCharacterBase::SetRelativeRotation(const FRotator& InRotation)
 	}
 
 	RelativeRotation = InRotation;
-	OnCharacterDataChanged();
+	OnCharacterDataChanged(EText3DRendererFlags::Layout);
 }
 
 void UText3DCharacterBase::SetRelativeScale(const FVector& InScale)
@@ -64,7 +69,32 @@ void UText3DCharacterBase::SetRelativeScale(const FVector& InScale)
 	}
 
 	RelativeScale = InScale;
-	OnCharacterDataChanged();
+	OnCharacterDataChanged(EText3DRendererFlags::Layout);
+}
+
+void UText3DCharacterBase::SetVisibility(bool bInVisibility)
+{
+	if (bVisible == bInVisibility)
+	{
+		return;
+	}
+
+	bVisible = bInVisibility;
+	OnCharacterDataChanged(EText3DRendererFlags::Visibility);
+}
+
+void UText3DCharacterBase::ResetCharacterState()
+{
+	const UText3DCharacterBase* CDO = GetDefault<UText3DCharacterBase>();
+#if WITH_EDITORONLY_DATA
+	Character = CDO->Character;
+#endif
+	RelativeLocation = CDO->RelativeLocation;
+	RelativeRotation = CDO->RelativeRotation;
+	RelativeScale = CDO->RelativeScale;
+	bVisible = CDO->bVisible;
+	Transform = CDO->Transform;
+	Metadata = CDO->Metadata;
 }
 
 #if WITH_EDITOR
@@ -72,32 +102,41 @@ void UText3DCharacterBase::PostEditUndo()
 {
 	Super::PostEditUndo();
 
-	OnCharacterDataChanged();
+	EText3DRendererFlags Flags = EText3DRendererFlags::Layout;
+	EnumAddFlags(Flags, EText3DRendererFlags::Visibility);
+
+	OnCharacterDataChanged(Flags);
 }
 
 void UText3DCharacterBase::PostEditChangeProperty(FPropertyChangedEvent& InEvent)
 {
 	Super::PostEditChangeProperty(InEvent);
 
-	static const TSet<FName> PropertyNames =
+	static const TSet<FName> LayoutPropertyNames =
 	{
 		GET_MEMBER_NAME_CHECKED(UText3DCharacterBase, RelativeLocation),
 		GET_MEMBER_NAME_CHECKED(UText3DCharacterBase, RelativeRotation),
 		GET_MEMBER_NAME_CHECKED(UText3DCharacterBase, RelativeScale)
 	};
 
-	if (PropertyNames.Contains(InEvent.GetMemberPropertyName()))
+	const FName MemberPropertyName = InEvent.GetMemberPropertyName();
+
+	if (LayoutPropertyNames.Contains(MemberPropertyName))
+	{
+		OnCharacterDataChanged(EText3DRendererFlags::Layout);
+	}
+	else if (MemberPropertyName == GET_MEMBER_NAME_CHECKED(UText3DCharacterBase, bVisible))
 	{
-		OnCharacterDataChanged();
+		OnCharacterDataChanged(EText3DRendererFlags::Visibility);
 	}
 }
 #endif
 
-void UText3DCharacterBase::OnCharacterDataChanged() const
+void UText3DCharacterBase::OnCharacterDataChanged(EText3DRendererFlags InFlags) const
 {
 	if (UText3DComponent* Component = GetTypedOuter<UText3DComponent>())
 	{
 		constexpr bool bImmediate = true;
-		Component->RequestUpdate(EText3DRendererFlags::Layout, bImmediate);
+		Component->RequestUpdate(InFlags, bImmediate);
 	}
 }
diff --git a/Engine/Plugins/Experimental/Text3D/Source/Text3D/Private/Characters/Text3DDefaultCharacter.cpp b/Engine/Plugins/Experimental/Text3D/Source/Text3D/Private/Characters/Text3DDefaultCharacter.cpp
--- a/Engine/Plugins/Experimental/Text3D/Source/Text3D/Private/Characters/Text3DDefaultCharacter.cpp
+++ b/Engine/Plugins/Experimental/Text3D/Source/Text3D/Private/Characters/Text3DDefaultCharacter.cpp
@@ -7,17 +7,6 @@ FName UText3DDefaultCharacter::GetKerningPropertyName()
 	return GET_MEMBER_NAME_CHECKED(UText3DDefaultCharacter, Kerning);
 }
 
-void UText3DDefaultCharacter::SetKerningOverride(bool bInKerningOverride)
-{
-	if (bKerningOverride == bInKerningOverride)
-	{
-		return;
-	}
-
-	bKerningOverride = bInKerningOverride;
-	OnCharacterDataChanged();
-}
-
 void UText3DDefaultCharacter::SetKerning(float InKerning)
 {
 	if (FMath::IsNearlyEqual(Kerning, InKerning))
@@ -26,7 +15,7 @@ void UText3DDefaultCharacter::SetKerning(float InKerning)
 	}
 
 	Kerning = InKerning;
-	OnCharacterDataChanged();
+	OnCharacterDataChanged(EText3DRendererFlags::Layout);
 }
 
 #if WITH_EDITOR
@@ -36,15 +25,22 @@ void UText3DDefaultCharacter::PostEditChangeProperty(FPropertyChangedEvent& InEv
 
 	const FName MemberName = InEvent.GetMemberPropertyName();
 
-	if (MemberName == GET_MEMBER_NAME_CHECKED(UText3DDefaultCharacter, Kerning)
-		|| MemberName == GET_MEMBER_NAME_CHECKED(UText3DDefaultCharacter, bKerningOverride))
+	if (MemberName == GET_MEMBER_NAME_CHECKED(UText3DDefaultCharacter, Kerning))
 	{
-		OnCharacterDataChanged();
+		OnCharacterDataChanged(EText3DRendererFlags::Layout);
 	}
 }
 #endif
 
-TOptional<float> UText3DDefaultCharacter::GetCharacterKerning() const
+float UText3DDefaultCharacter::GetCharacterKerning() const
 {
-	return bKerningOverride ? Kerning : TOptional<float>();
+	return Kerning;
+}
+
+void UText3DDefaultCharacter::ResetCharacterState()
+{
+	Super::ResetCharacterState();
+
+	const UText3DDefaultCharacter* CDO = GetDefault<UText3DDefaultCharacter>();
+	Kerning = CDO->Kerning;
 }
diff --git a/Engine/Plugins/Experimental/Text3D/Source/Text3D/Private/Extensions/Text3DDefaultCharacterExtension.cpp b/Engine/Plugins/Experimental/Text3D/Source/Text3D/Private/Extensions/Text3DDefaultCharacterExtension.cpp
--- a/Engine/Plugins/Experimental/Text3D/Source/Text3D/Private/Extensions/Text3DDefaultCharacterExtension.cpp
+++ b/Engine/Plugins/Experimental/Text3D/Source/Text3D/Private/Extensions/Text3DDefaultCharacterExtension.cpp
@@ -35,7 +35,7 @@ EText3DExtensionResult UText3DDefaultCharacterExtension::PreRendererUpdate(EText
 		return EText3DExtensionResult::Active;
 	}
 
-	UText3DComponent* Text3DComponent = GetText3DComponent();
+	const UText3DComponent* Text3DComponent = GetText3DComponent();
 	const FTypefaceEntry* TypeFaceEntry = Text3DComponent->GetTypeFaceEntry();
 
 	if (!TypeFaceEntry)
@@ -77,7 +77,11 @@ void UText3DDefaultCharacterExtension::AllocateTextCharacters(uint16 InCharacter
 	{
 		for (int32 CharacterIndex = InCharacterCount; CharacterIndex < TextCharacters.Num(); CharacterIndex++)
 		{
-			TextCharactersPool.Add(TextCharacters[CharacterIndex]);
+			if (UText3DCharacterBase* Character = TextCharacters[CharacterIndex])
+			{
+				Character->ResetCharacterState();
+				TextCharactersPool.Add(Character);
+			}
 		}
 	}
 	else if (RemainingCharacterCount < 0)
diff --git a/Engine/Plugins/Experimental/Text3D/Source/Text3D/Private/Extensions/Text3DDefaultLayoutExtension.cpp b/Engine/Plugins/Experimental/Text3D/Source/Text3D/Private/Extensions/Text3DDefaultLayoutExtension.cpp
--- a/Engine/Plugins/Experimental/Text3D/Source/Text3D/Private/Extensions/Text3DDefaultLayoutExtension.cpp
+++ b/Engine/Plugins/Experimental/Text3D/Source/Text3D/Private/Extensions/Text3DDefaultLayoutExtension.cpp
@@ -45,14 +45,14 @@ FName UText3DDefaultLayoutExtension::GetScaleProportionallyPropertyName()
 	return GET_MEMBER_NAME_CHECKED(UText3DDefaultLayoutExtension, bScaleProportionally);
 }
 
-void UText3DDefaultLayoutExtension::SetKerning(const float Value)
+void UText3DDefaultLayoutExtension::SetTracking(const float Value)
 {
-	if (FMath::IsNearlyEqual(Kerning, Value))
+	if (FMath::IsNearlyEqual(Tracking, Value))
 	{
 		return;
 	}
 	
-	Kerning = Value;
+	Tracking = Value;
 	OnLayoutOptionsChanged();
 }
 
@@ -271,7 +271,7 @@ EText3DExtensionResult UText3DDefaultLayoutExtension::PreRendererUpdate(EText3DR
 	ShapedText->GlyphSizes.SetNum(Text3DComponent->GetCharacterCount());
 	Text3DComponent->ForEachCharacter([this](UText3DCharacterBase* InCharacter, uint16 InIndex, uint16 InCount)
 	{
-		ShapedText->Kernings[InIndex] = InCharacter->GetCharacterKerning().Get(Kerning);
+		ShapedText->Kernings[InIndex] = InCharacter->GetCharacterKerning();
 		ShapedText->GlyphSizes[InIndex] = InCharacter->GetCharacterMetadata().MeshBounds.GetSize();
 		ShapedText->CharacterGlyphMap.Add(InCharacter->GetCharacterMetadata().LayoutIndex, InIndex);
 	});
@@ -279,7 +279,7 @@ EText3DExtensionResult UText3DDefaultLayoutExtension::PreRendererUpdate(EText3DR
 	UText3DGeometryExtensionBase* GeometryExtension = Text3DComponent->GetGeometryExtension();
 	ShapedText->PivotHAlignment = GeometryExtension->GetGlyphHAlignment();
 	ShapedText->PivotVAlignment = GeometryExtension->GetGlyphVAlignment();
-	ShapedText->Kerning = Kerning;
+	ShapedText->Tracking = Tracking;
 	ShapedText->WordSpacing = WordSpacing;
 	ShapedText->MaxWidth = MaxWidth;
 	ShapedText->bWrap = bUseMaxWidth && MaxWidthBehavior == EText3DMaxWidthHandling::WrapAndScale;
@@ -317,9 +317,9 @@ EText3DExtensionResult UText3DDefaultLayoutExtension::PreRendererUpdate(EText3DR
 					break;
 				}
 
-				if (ShapedGlyph.GetIndex() != 0)
+				if (ShapedGlyph.GetIndex() != 0 && ShapedGlyph.GetIndex() != ShapedLine->GlyphsToRender.Num() - 1)
 				{
-					Shift.Y += ShapedText->Kernings[*CharacterIndex];
+					Shift.Y += Tracking + ShapedText->Kernings[*CharacterIndex];
 				}
 
 				CharacterTransform.Accumulate(FTransform(Location + Shift));
@@ -431,7 +431,7 @@ void UText3DDefaultLayoutExtension::PostEditChangeProperty(FPropertyChangedEvent
 
 	static const TSet<FName> LayoutPropertyNames
 	{
-		GET_MEMBER_NAME_CHECKED(UText3DDefaultLayoutExtension, Kerning),
+		GET_MEMBER_NAME_CHECKED(UText3DDefaultLayoutExtension, Tracking),
 		GET_MEMBER_NAME_CHECKED(UText3DDefaultLayoutExtension, LineSpacing),
 		GET_MEMBER_NAME_CHECKED(UText3DDefaultLayoutExtension, WordSpacing),
 		GET_MEMBER_NAME_CHECKED(UText3DDefaultLayoutExtension, HorizontalAlignment),
diff --git a/Engine/Plugins/Experimental/Text3D/Source/Text3D/Private/LayoutBuilders/Text3DShapedGlyphLine.cpp b/Engine/Plugins/Experimental/Text3D/Source/Text3D/Private/LayoutBuilders/Text3DShapedGlyphLine.cpp
--- a/Engine/Plugins/Experimental/Text3D/Source/Text3D/Private/LayoutBuilders/Text3DShapedGlyphLine.cpp
+++ b/Engine/Plugins/Experimental/Text3D/Source/Text3D/Private/LayoutBuilders/Text3DShapedGlyphLine.cpp
@@ -7,18 +7,17 @@ float FShapedGlyphLine::GetWidthAdvance(const int32 InIndex, const float InKerni
 	check(InIndex >= 0 && InIndex < GlyphsToRender.Num());
 
 	const FShapedGlyphEntry& Glyph = GlyphsToRender[InIndex];
-
 	float Advance = Glyph.XOffset + Glyph.XAdvance;
 
-	// Skip first character kerning
-	if (InIndex != 0)
+	const bool bFirstCharacter = InIndex == 0;
+	const bool bLastCharacter = InIndex == GlyphsToRender.Num() - 1;
+
+	// Skip first and last character for kerning and tracking
+	if (!bFirstCharacter && !bLastCharacter)
 	{
 		// @note: as per FSlateElementBatcher::BuildShapedTextSequence, per Glyph Kerning isn't used
 		Advance += InKerning;
-	}
 
-	if (InIndex < GlyphsToRender.Num() - 1)
-	{
 		if (!Glyph.bIsVisible)
 		{
 			Advance += InWordSpacing;
diff --git a/Engine/Plugins/Experimental/Text3D/Source/Text3D/Private/LayoutBuilders/Text3DShapedGlyphText.cpp b/Engine/Plugins/Experimental/Text3D/Source/Text3D/Private/LayoutBuilders/Text3DShapedGlyphText.cpp
--- a/Engine/Plugins/Experimental/Text3D/Source/Text3D/Private/LayoutBuilders/Text3DShapedGlyphText.cpp
+++ b/Engine/Plugins/Experimental/Text3D/Source/Text3D/Private/LayoutBuilders/Text3DShapedGlyphText.cpp
@@ -15,7 +15,7 @@ void FText3DShapedGlyphText::Reset()
 	PivotVAlignment = EText3DVerticalTextAlignment::Bottom;
 	Kernings.Reset();
 	GlyphSizes.Reset();
-	Kerning = 0.f;
+	Tracking = 0.f;
 	WordSpacing = 0.0f;
 	bWrap = false;
 	Lines.Reset();
@@ -41,8 +41,8 @@ void FText3DShapedGlyphText::CalculateWidth()
 		for (TConstEnumerateRef<FShapedGlyphEntry> CurrentGlyph : EnumerateRange(GlyphLine->GlyphsToRender))
 		{
 			const int32* CharacterIndex = CharacterGlyphMap.Find(LayoutIndex);
-			const float CurrentKerning = !!CharacterIndex ? Kernings[*CharacterIndex] : Kerning;
-			const float GlyphAdv = GlyphLine->GetWidthAdvance(CurrentGlyph.GetIndex(), CurrentKerning, WordSpacing);
+			const float CurrentKerning = !!CharacterIndex ? Kernings[*CharacterIndex] : 0.f;
+			const float GlyphAdv = GlyphLine->GetWidthAdvance(CurrentGlyph.GetIndex(), Tracking + CurrentKerning, WordSpacing);
 
 			// If we're at the end the line or at whitespace
 			const bool bWordBreak = !CurrentGlyph->bIsVisible || CurrentGlyph.GetIndex() == GlyphLine->GlyphsToRender.Num() - 1;
diff --git a/Engine/Plugins/Experimental/Text3D/Source/Text3D/Private/LayoutBuilders/Text3DShapedGlyphText.h b/Engine/Plugins/Experimental/Text3D/Source/Text3D/Private/LayoutBuilders/Text3DShapedGlyphText.h
--- a/Engine/Plugins/Experimental/Text3D/Source/Text3D/Private/LayoutBuilders/Text3DShapedGlyphText.h
+++ b/Engine/Plugins/Experimental/Text3D/Source/Text3D/Private/LayoutBuilders/Text3DShapedGlyphText.h
@@ -31,8 +31,8 @@ struct FText3DShapedGlyphText
 	/** Individual character size */
 	TArray<FVector> GlyphSizes;
 
-	/** General kerning for invisible character */
-	float Kerning;
+	/** General tracking for all characters */
+	float Tracking;
 
 	/** Spacing between words */
 	float WordSpacing;
diff --git a/Engine/Plugins/Experimental/Text3D/Source/Text3D/Private/Renderers/StaticMeshes/Text3DStaticMeshesRenderer.cpp b/Engine/Plugins/Experimental/Text3D/Source/Text3D/Private/Renderers/StaticMeshes/Text3DStaticMeshesRenderer.cpp
--- a/Engine/Plugins/Experimental/Text3D/Source/Text3D/Private/Renderers/StaticMeshes/Text3DStaticMeshesRenderer.cpp
+++ b/Engine/Plugins/Experimental/Text3D/Source/Text3D/Private/Renderers/StaticMeshes/Text3DStaticMeshesRenderer.cpp
@@ -72,8 +72,9 @@ void UText3DStaticMeshesRenderer::OnUpdate(EText3DRendererFlags InFlags)
 
 		const UText3DMaterialExtensionBase* MaterialExtension = TextComponent->GetMaterialExtension();
 
-		for (UStaticMeshComponent* StaticMeshComponent : CharacterMeshes)
+		TextComponent->ForEachCharacter([this, MaterialExtension](UText3DCharacterBase* InCharacter, uint16 InIndex, uint16 InCount)
 		{
+			UStaticMeshComponent* StaticMeshComponent = CharacterMeshes[InIndex];
 			for (int32 GroupIndex = 0; GroupIndex < static_cast<int32>(EText3DGroupType::TypeCount); GroupIndex++)
 			{
 				const int32 MaterialIndex = StaticMeshComponent->GetMaterialIndex(SlotNames[GroupIndex]);
@@ -90,23 +91,24 @@ void UText3DStaticMeshesRenderer::OnUpdate(EText3DRendererFlags InFlags)
 					StaticMeshComponent->SetMaterial(MaterialIndex, Material);
 				}
 			}
-		}
+		});
 	}
 
 	if (EnumHasAnyFlags(InFlags, EText3DRendererFlags::Visibility))
 	{
 		const UText3DRenderingExtensionBase* RenderingExtension = TextComponent->GetRenderingExtension();
 
-		for (UStaticMeshComponent* StaticMeshComponent : CharacterMeshes)
+		TextComponent->ForEachCharacter([this, TextComponent, RenderingExtension](UText3DCharacterBase* InCharacter, uint16 InIndex, uint16 InCount)
 		{
+			UStaticMeshComponent* StaticMeshComponent = CharacterMeshes[InIndex];
 			StaticMeshComponent->SetHiddenInGame(TextComponent->bHiddenInGame);
-			StaticMeshComponent->SetVisibility(TextComponent->GetVisibleFlag());
+			StaticMeshComponent->SetVisibility(TextComponent->GetVisibleFlag() && InCharacter->GetVisibility());
 			StaticMeshComponent->SetCastShadow(RenderingExtension->GetTextCastShadow());
 			StaticMeshComponent->SetCastHiddenShadow(RenderingExtension->GetTextCastHiddenShadow());
 			StaticMeshComponent->SetAffectDynamicIndirectLighting(RenderingExtension->GetTextAffectDynamicIndirectLighting());
 			StaticMeshComponent->SetAffectIndirectLightingWhileHidden(RenderingExtension->GetTextAffectIndirectLightingWhileHidden());
 			StaticMeshComponent->SetHoldout(RenderingExtension->GetTextHoldout());
-		}
+		});
 	}
 }
 
diff --git a/Engine/Plugins/Experimental/Text3D/Source/Text3D/Private/Text3DComponent.cpp b/Engine/Plugins/Experimental/Text3D/Source/Text3D/Private/Text3DComponent.cpp
--- a/Engine/Plugins/Experimental/Text3D/Source/Text3D/Private/Text3DComponent.cpp
+++ b/Engine/Plugins/Experimental/Text3D/Source/Text3D/Private/Text3DComponent.cpp
@@ -294,7 +294,7 @@ void UText3DComponent::Serialize(FArchive& InArchive)
 			{
 				DefaultLayoutExtension->SetHorizontalAlignment(HorizontalAlignment_DEPRECATED);
 				DefaultLayoutExtension->SetVerticalAlignment(VerticalAlignment_DEPRECATED);
-				DefaultLayoutExtension->SetKerning(Kerning_DEPRECATED);
+				DefaultLayoutExtension->SetTracking(Kerning_DEPRECATED);
 				DefaultLayoutExtension->SetLineSpacing(LineSpacing_DEPRECATED);
 				DefaultLayoutExtension->SetWordSpacing(WordSpacing_DEPRECATED);
 				DefaultLayoutExtension->SetUseMaxWidth(bHasMaxWidth_DEPRECATED);
@@ -553,7 +553,7 @@ float UText3DComponent::GetKerning() const
 {
 	if (const UText3DDefaultLayoutExtension* DefaultLayoutExtension = GetCastedLayoutExtension<UText3DDefaultLayoutExtension>())
 	{
-		return DefaultLayoutExtension->GetKerning();
+		return DefaultLayoutExtension->GetTracking();
 	}
 
 	return 0.f;
@@ -587,7 +587,7 @@ void UText3DComponent::SetKerning(const float Value)
 {
 	if (UText3DDefaultLayoutExtension* DefaultLayoutExtension = GetCastedLayoutExtension<UText3DDefaultLayoutExtension>())
 	{
-		DefaultLayoutExtension->SetKerning(Value);
+		DefaultLayoutExtension->SetTracking(Value);
 	}
 }
 
diff --git a/Engine/Plugins/Experimental/Text3D/Source/Text3D/Public/Characters/Text3DCharacterBase.h b/Engine/Plugins/Experimental/Text3D/Source/Text3D/Public/Characters/Text3DCharacterBase.h
--- a/Engine/Plugins/Experimental/Text3D/Source/Text3D/Public/Characters/Text3DCharacterBase.h
+++ b/Engine/Plugins/Experimental/Text3D/Source/Text3D/Public/Characters/Text3DCharacterBase.h
@@ -16,6 +16,7 @@ public:
 	static TEXT3D_API FName GetRelativeLocationPropertyName();
 	static TEXT3D_API FName GetRelativeRotationPropertyName();
 	static TEXT3D_API FName GetRelativeScalePropertyName();
+	static TEXT3D_API FName GetVisiblePropertyName();
 
 	TEXT3D_API FTransform& GetTransform(bool bInReset);
 
@@ -62,12 +63,22 @@ public:
 		return RelativeScale;
 	}
 
-	/** Get character custom kerning if supported */
-	virtual TOptional<float> GetCharacterKerning() const
+	UFUNCTION()
+	TEXT3D_API void SetVisibility(bool bInVisibility);
+	bool GetVisibility() const
+	{
+		return bVisible;
+	}
+
+	/** Get character custom kerning */
+	virtual float GetCharacterKerning() const
 	{
-		return TOptional<float>();
+		return 0.f;
 	}
 
+	/** Reset properties to their initial state when character is recycled */
+	TEXT3D_API virtual void ResetCharacterState();
+
 protected:
 	//~ Begin UObject
 #if WITH_EDITOR
@@ -76,7 +87,7 @@ protected:
 #endif
 	//~ End UObject
 
-	void OnCharacterDataChanged() const;
+	void OnCharacterDataChanged(EText3DRendererFlags InFlags) const;
 
 #if WITH_EDITORONLY_DATA
 	/** Used for TitleProperty of the array to display this instead of the index */
@@ -93,6 +104,9 @@ protected:
 	UPROPERTY(EditInstanceOnly, Setter, Getter, Category="Character", meta=(Delta="0.05", MotionDesignVectorWidget, AllowPreserveRatio="XYZ", ClampMin="0", AllowPrivateAccess = "true"))
 	FVector RelativeScale = FVector::OneVector;
 
+	UPROPERTY(EditInstanceOnly, Setter="SetVisibility", Getter="GetVisibility", Category="Character", meta=(AllowPrivateAccess = "true"))
+	bool bVisible = true;
+
 	/** Final transform after all extensions are applied */
 	FTransform Transform = FTransform::Identity;
 
diff --git a/Engine/Plugins/Experimental/Text3D/Source/Text3D/Public/Characters/Text3DDefaultCharacter.h b/Engine/Plugins/Experimental/Text3D/Source/Text3D/Public/Characters/Text3DDefaultCharacter.h
--- a/Engine/Plugins/Experimental/Text3D/Source/Text3D/Public/Characters/Text3DDefaultCharacter.h
+++ b/Engine/Plugins/Experimental/Text3D/Source/Text3D/Public/Characters/Text3DDefaultCharacter.h
@@ -15,12 +15,6 @@ class UText3DDefaultCharacter : public UText3DCharacterBase
 public:
 	static TEXT3D_API FName GetKerningPropertyName();
 
-	TEXT3D_API void SetKerningOverride(bool bInKerningOverride);
-	bool GetKerningOverride() const
-	{
-		return bKerningOverride;
-	}
-
 	TEXT3D_API void SetKerning(float InKerning);
 	float GetKerning() const
 	{
@@ -35,13 +29,11 @@ protected:
 	//~ End UObject
 
 	//~ Begin UText3DCharacterBase
-	virtual TOptional<float> GetCharacterKerning() const override;
+	virtual float GetCharacterKerning() const override;
+	virtual void ResetCharacterState() override;
 	//~ End UText3DCharacterBase
 
-	UPROPERTY(EditInstanceOnly, Getter="GetKerningOverride", Setter="SetKerningOverride", Category="Character", meta=(AllowPrivateAccess = "true"))
-	bool bKerningOverride = false;
-
 	/** Kerning adjusts the space between this character */
-	UPROPERTY(EditInstanceOnly, Getter, Setter, Category="Character", meta=(EditCondition="bKerningOverride", AllowPrivateAccess = "true"))
+	UPROPERTY(EditInstanceOnly, Getter, Setter, Category="Character", meta=(AllowPrivateAccess = "true"))
 	float Kerning = 0.f;
 };
diff --git a/Engine/Plugins/Experimental/Text3D/Source/Text3D/Public/Extensions/Text3DDefaultLayoutExtension.h b/Engine/Plugins/Experimental/Text3D/Source/Text3D/Public/Extensions/Text3DDefaultLayoutExtension.h
--- a/Engine/Plugins/Experimental/Text3D/Source/Text3D/Public/Extensions/Text3DDefaultLayoutExtension.h
+++ b/Engine/Plugins/Experimental/Text3D/Source/Text3D/Public/Extensions/Text3DDefaultLayoutExtension.h
@@ -23,16 +23,16 @@ public:
 	static TEXT3D_API FName GetMaxWidthPropertyName();
 	static TEXT3D_API FName GetScaleProportionallyPropertyName();
 
-	/** Get the kerning value and signal the primitives to be rebuilt */
+	/** Get the tracking value and signal the primitives to be rebuilt */
 	UFUNCTION(BlueprintPure, Category = "Text3D|Layout")
-	float GetKerning() const
+	float GetTracking() const
 	{
-		return Kerning;
+		return Tracking;
 	}
 
-	/** Set the kerning value and signal the primitives to be rebuilt */
+	/** Set the tracking value and signal the primitives to be rebuilt */
 	UFUNCTION(BlueprintCallable, Category = "Text3D|Layout")
-	void SetKerning(const float Value);
+	void SetTracking(const float Value);
 
 	/** Get the line spacing value and signal the primitives to be rebuilt */
 	UFUNCTION(BlueprintPure, Category = "Text3D|Layout")
@@ -173,9 +173,9 @@ protected:
 	UPROPERTY(EditAnywhere, Getter, Setter, Category = "Layout", meta = (AllowPrivateAccess = "true", CustomAlignmentWidget))
 	EText3DVerticalTextAlignment VerticalAlignment = EText3DVerticalTextAlignment::FirstLine;
 
-	/** Text kerning */
+	/** Text tracking affects all characters */
 	UPROPERTY(EditAnywhere, Getter, Setter, Category = "Layout", meta = (AllowPrivateAccess = "true"))
-	float Kerning = 0.0f;
+	float Tracking = 0.0f;
 
 	/** Extra line spacing */
 	UPROPERTY(EditAnywhere, Getter, Setter, Category = "Layout", meta = (AllowPrivateAccess = "true"))
