MotionDesign : Text3D
- Reworking default text transform extension effect to support ease curve

diff --git a/Engine/Plugins/Experimental/Avalanche/Source/AvalancheText/Private/AvaText3DComponent.cpp b/Engine/Plugins/Experimental/Avalanche/Source/AvalancheText/Private/AvaText3DComponent.cpp
--- a/Engine/Plugins/Experimental/Avalanche/Source/AvalancheText/Private/AvaText3DComponent.cpp
+++ b/Engine/Plugins/Experimental/Avalanche/Source/AvalancheText/Private/AvaText3DComponent.cpp
@@ -124,21 +124,17 @@ void UAvaText3DComponent::Serialize(FArchive& InArchive)
 					TransformExtension->SetLocationEnabled(CharacterTransform->GetLocationEnabled());
 					TransformExtension->SetLocationProgress(CharacterTransform->GetLocationProgress());
 					TransformExtension->SetLocationOrder(CharacterTransform->GetLocationOrder());
-					TransformExtension->SetLocationRange(CharacterTransform->GetLocationRange());
-					TransformExtension->SetLocationDistance(CharacterTransform->GetLocationDistance());
 					TransformExtension->SetLocationDistance(CharacterTransform->GetLocationDistance());
 
 					TransformExtension->SetRotationEnabled(CharacterTransform->GetRotationEnabled());
 					TransformExtension->SetRotationProgress(CharacterTransform->GetRotationProgress());
 					TransformExtension->SetRotationOrder(CharacterTransform->GetRotationOrder());
-					TransformExtension->SetRotationRange(CharacterTransform->GetRotationRange());
 					TransformExtension->SetRotationBegin(CharacterTransform->GetRotationBegin());
 					TransformExtension->SetRotationEnd(CharacterTransform->GetRotationEnd());
 
 					TransformExtension->SetScaleEnabled(CharacterTransform->GetScaleEnabled());
 					TransformExtension->SetScaleProgress(CharacterTransform->GetScaleProgress());
 					TransformExtension->SetScaleOrder(CharacterTransform->GetScaleOrder());
-					TransformExtension->SetScaleRange(CharacterTransform->GetScaleRange());
 					TransformExtension->SetScaleBegin(CharacterTransform->GetScaleBegin());
 					TransformExtension->SetScaleEnd(CharacterTransform->GetScaleEnd());
 
diff --git a/Engine/Plugins/Experimental/Text3D/Source/Text3D/Private/Extensions/Text3DDefaultTransformExtension.cpp b/Engine/Plugins/Experimental/Text3D/Source/Text3D/Private/Extensions/Text3DDefaultTransformExtension.cpp
--- a/Engine/Plugins/Experimental/Text3D/Source/Text3D/Private/Extensions/Text3DDefaultTransformExtension.cpp
+++ b/Engine/Plugins/Experimental/Text3D/Source/Text3D/Private/Extensions/Text3DDefaultTransformExtension.cpp
@@ -2,6 +2,8 @@
 
 #include "Extensions/Text3DDefaultTransformExtension.h"
 
+#include "Curves/CurveFloat.h"
+
 void UText3DDefaultTransformExtension::SetLocationEnabled(bool bEnabled)
 {
 	if (bLocationEnabled != bEnabled)
@@ -30,23 +32,24 @@ void UText3DDefaultTransformExtension::SetLocationOrder(EText3DCharacterEffectOr
 	}
 }
 
-void UText3DDefaultTransformExtension::SetLocationRange(float Range)
+void UText3DDefaultTransformExtension::SetLocationDistance(const FVector& Distance)
 {
-	Range = FMath::Clamp(Range, 0.0f, 100.0f);
-	if (!FMath::IsNearlyEqual(LocationRange, Range))
+	if (!(LocationDistance - Distance).IsNearlyZero())
 	{
-		LocationRange = Range;
+		LocationDistance = Distance;
 		OnTransformOptionsChanged();
 	}
 }
 
-void UText3DDefaultTransformExtension::SetLocationDistance(const FVector& Distance)
+void UText3DDefaultTransformExtension::SetLocationEaseCurve(UCurveFloat* InEaseCurve)
 {
-	if (!(LocationDistance - Distance).IsNearlyZero())
+	if (LocationEaseCurve == InEaseCurve)
 	{
-		LocationDistance = Distance;
-		OnTransformOptionsChanged();
+		return;
 	}
+
+	LocationEaseCurve = InEaseCurve;
+	OnTransformOptionsChanged();
 }
 
 void UText3DDefaultTransformExtension::SetScaleEnabled(bool bEnabled)
@@ -77,16 +80,6 @@ void UText3DDefaultTransformExtension::SetScaleOrder(EText3DCharacterEffectOrder
 	}
 }
 
-void UText3DDefaultTransformExtension::SetScaleRange(float Range)
-{
-	Range = FMath::Clamp(Range, 0.0f, 100.0f);
-	if (!FMath::IsNearlyEqual(ScaleRange, Range))
-	{
-		ScaleRange = Range;
-		OnTransformOptionsChanged();
-	}
-}
-
 void UText3DDefaultTransformExtension::SetScaleBegin(const FVector& Value)
 {
 	if (!(ScaleBegin - Value).IsNearlyZero())
@@ -105,6 +98,17 @@ void UText3DDefaultTransformExtension::SetScaleEnd(const FVector& Value)
 	}
 }
 
+void UText3DDefaultTransformExtension::SetScaleEaseCurve(UCurveFloat* InEaseCurve)
+{
+	if (ScaleEaseCurve == InEaseCurve)
+	{
+		return;
+	}
+
+	ScaleEaseCurve = InEaseCurve;
+	OnTransformOptionsChanged();
+}
+
 void UText3DDefaultTransformExtension::SetRotationEnabled(bool bEnabled)
 {
 	if (bRotationEnabled != bEnabled)
@@ -133,16 +137,6 @@ void UText3DDefaultTransformExtension::SetRotationOrder(EText3DCharacterEffectOr
 	}
 }
 
-void UText3DDefaultTransformExtension::SetRotationRange(float Range)
-{
-	Range = FMath::Clamp(Range, 0.0f, 100.0f);
-	if (!FMath::IsNearlyEqual(RotationRange, Range))
-	{
-		RotationRange = Range;
-		OnTransformOptionsChanged();
-	}
-}
-
 void UText3DDefaultTransformExtension::SetRotationBegin(const FRotator& Value)
 {
 	if (!(RotationBegin - Value).IsNearlyZero())
@@ -161,6 +155,17 @@ void UText3DDefaultTransformExtension::SetRotationEnd(const FRotator& Value)
 	}
 }
 
+void UText3DDefaultTransformExtension::SetRotationEaseCurve(UCurveFloat* InEaseCurve)
+{
+	if (RotationEaseCurve == InEaseCurve)
+	{
+		return;
+	}
+
+	RotationEaseCurve = InEaseCurve;
+	OnTransformOptionsChanged();
+}
+
 UText3DDefaultTransformExtension::UText3DDefaultTransformExtension()
 {
 	OnTransformOptionsChanged();
@@ -176,21 +181,21 @@ void UText3DDefaultTransformExtension::PostEditChangeProperty(FPropertyChangedEv
 		GET_MEMBER_NAME_CHECKED(UText3DDefaultTransformExtension, bLocationEnabled),
 		GET_MEMBER_NAME_CHECKED(UText3DDefaultTransformExtension, LocationProgress),
 		GET_MEMBER_NAME_CHECKED(UText3DDefaultTransformExtension, LocationOrder),
-		GET_MEMBER_NAME_CHECKED(UText3DDefaultTransformExtension, LocationRange),
 		GET_MEMBER_NAME_CHECKED(UText3DDefaultTransformExtension, LocationDistance),
+		GET_MEMBER_NAME_CHECKED(UText3DDefaultTransformExtension, LocationEaseCurve),
 		GET_MEMBER_NAME_CHECKED(UText3DDefaultTransformExtension, bRotationEnabled),
 		GET_MEMBER_NAME_CHECKED(UText3DDefaultTransformExtension, RotationProgress),
 		GET_MEMBER_NAME_CHECKED(UText3DDefaultTransformExtension, RotationOrder),
-		GET_MEMBER_NAME_CHECKED(UText3DDefaultTransformExtension, RotationRange),
 		GET_MEMBER_NAME_CHECKED(UText3DDefaultTransformExtension, RotationBegin),
 		GET_MEMBER_NAME_CHECKED(UText3DDefaultTransformExtension, RotationEnd),
 		GET_MEMBER_NAME_CHECKED(UText3DDefaultTransformExtension, RotationEnd),
+		GET_MEMBER_NAME_CHECKED(UText3DDefaultTransformExtension, RotationEaseCurve),
 		GET_MEMBER_NAME_CHECKED(UText3DDefaultTransformExtension, bScaleEnabled),
 		GET_MEMBER_NAME_CHECKED(UText3DDefaultTransformExtension, ScaleProgress),
 		GET_MEMBER_NAME_CHECKED(UText3DDefaultTransformExtension, ScaleOrder),
-		GET_MEMBER_NAME_CHECKED(UText3DDefaultTransformExtension, ScaleRange),
 		GET_MEMBER_NAME_CHECKED(UText3DDefaultTransformExtension, ScaleBegin),
 		GET_MEMBER_NAME_CHECKED(UText3DDefaultTransformExtension, ScaleEnd),
+		GET_MEMBER_NAME_CHECKED(UText3DDefaultTransformExtension, ScaleEaseCurve),
 	};
 
 	if (PropertyNames.Contains(InEvent.GetMemberPropertyName()))
@@ -209,56 +214,32 @@ void UText3DDefaultTransformExtension::OnTransformOptionsChanged()
 
 void UText3DDefaultTransformExtension::ApplyTransform(uint32 InGlyphIndex, uint32 InGlyphCount, FTransform& InGlyphTransform)
 {
+	InGlyphIndex += 1;
+	InGlyphCount += 1;
+
 	if (bLocationEnabled)
 	{
-		float Effect = GetEffectValue(InGlyphIndex, InGlyphCount, LocationOrder, LocationProgress, LocationRange);
-		FVector Location = LocationDistance * Effect;
+		const float Effect = CalculateEffect(InGlyphIndex, InGlyphCount, LocationOrder, LocationProgress, LocationEaseCurve);
+		const FVector Location = LocationDistance * Effect;
 		InGlyphTransform.AddToTranslation(Location);
 	}
 
 	if (bScaleEnabled)
 	{
-		float Effect = GetEffectValue(InGlyphIndex, InGlyphCount, ScaleOrder, ScaleProgress, ScaleRange);
-		FVector Scale = ScaleBegin + ((ScaleEnd - ScaleBegin) * Effect);
+		const float Effect = CalculateEffect(InGlyphIndex, InGlyphCount, ScaleOrder, ScaleProgress, ScaleEaseCurve);
+		const FVector Scale = ScaleBegin + ((ScaleEnd - ScaleBegin) * Effect);
 		InGlyphTransform.MultiplyScale3D(Scale);
 	}
 
 	if (bRotationEnabled)
 	{
-		float Effect = GetEffectValue(InGlyphIndex, InGlyphCount, RotationOrder, RotationProgress, RotationRange);
-		FRotator Rotator = RotationBegin + (RotationEnd - RotationBegin) * Effect;
+		const float Effect = CalculateEffect(InGlyphIndex, InGlyphCount, RotationOrder, RotationProgress, RotationEaseCurve);
+		const FRotator Rotator = RotationBegin + (RotationEnd - RotationBegin) * Effect;
 		InGlyphTransform.ConcatenateRotation(Rotator.Quaternion());
 	}
 }
 
-float UText3DDefaultTransformExtension::GetEffectValue(int32 Index, int32 Total, EText3DCharacterEffectOrder Order, float Progress, float Range)
-{
-	float Effect = 0.0f, Strip = 1.0f;
-	GetLineParameters(Range * 0.01f, Order, Total, Effect, Strip);
-
-	int32 Position = GetEffectPosition(Index, Total, Order);
-	return FMath::Clamp((Progress * 0.01f - Effect * Position) / Strip, 0.0f, 1.0f);
-}
-
-void UText3DDefaultTransformExtension::GetLineParameters(float Range, EText3DCharacterEffectOrder Order, int32 Total, float & EffectOut, float & StripOut)
-{
-	float Segment = 1.0f / Total;
-	EffectOut = (1.0f - Range) / Total;
-	StripOut = FMath::Max(Range, Segment);
-
-	if (Order == EText3DCharacterEffectOrder::FromCenter || Order == EText3DCharacterEffectOrder::ToCenter)
-	{
-		int32 Center = Total / 2.0f - 0.5f;
-		EffectOut *= (float)Total / (Center + 1.0f);
-	}
-
-	if (FMath::IsNearlyZero(EffectOut))
-	{
-		EffectOut = 0.01f;
-	}
-}
-
-int32 UText3DDefaultTransformExtension::GetEffectPosition(int32 Index, int32 Total, EText3DCharacterEffectOrder Order)
+int32 UText3DDefaultTransformExtension::GetEffectPosition(int32 Index, int32 Total, EText3DCharacterEffectOrder Order) const
 {
 	int32 Center = 0.5f * Total - 0.5f;
 	switch (Order)
@@ -291,4 +272,26 @@ int32 UText3DDefaultTransformExtension::GetEffectPosition(int32 Index, int32 Tot
 	}
 
 	return Index;
-}
\ No newline at end of file
+}
+
+float UText3DDefaultTransformExtension::CalculateEffect(int32 InIndex, int32 InTotal, EText3DCharacterEffectOrder InOrder, float InProgress, const UCurveFloat* InEaseCurve) const
+{
+	const int32 Position = GetEffectPosition(InIndex, InTotal, InOrder);
+	const float NormalizedProgress = FMath::Clamp(InProgress * 0.01f, 0.0f, 1.0f);
+	const float StaggerAmount = 1.0f / FMath::Max(1, InTotal);
+	const float CharacterStart = Position * StaggerAmount;
+	const float CharacterProgress = FMath::Clamp((NormalizedProgress - CharacterStart) / (1.0f - CharacterStart), 0.0f, 1.0f);
+	float CharacterEased = CharacterProgress;
+
+	if (InEaseCurve)
+	{
+		float StartTime = 0.0f;
+		float EndTime = 0.0f;
+		InEaseCurve->FloatCurve.GetTimeRange(StartTime, EndTime);
+
+		const float NormalizedTime = FMath::GetMappedRangeValueClamped(FVector2D(0.f, 1.f), FVector2D(StartTime, EndTime), CharacterProgress);
+		CharacterEased = InEaseCurve->FloatCurve.Eval(NormalizedTime);
+	}
+
+	return CharacterEased;
+}
diff --git a/Engine/Plugins/Experimental/Text3D/Source/Text3D/Private/Renderers/Text3DRendererBase.cpp b/Engine/Plugins/Experimental/Text3D/Source/Text3D/Private/Renderers/Text3DRendererBase.cpp
--- a/Engine/Plugins/Experimental/Text3D/Source/Text3D/Private/Renderers/Text3DRendererBase.cpp
+++ b/Engine/Plugins/Experimental/Text3D/Source/Text3D/Private/Renderers/Text3DRendererBase.cpp
@@ -145,7 +145,7 @@ void UText3DRendererBase::Clear()
 	CachedBounds.Reset();
 
 	const AActor* Owner = Text3DComponent->GetOwner();
-	UE_LOG(LogText3D, Log, TEXT("%s : Text3DRenderer %s Cleared"), !!Owner ? *Owner->GetActorNameOrLabel() : TEXT("Invalid owner"), *GetName().ToString())
+	UE_LOG(LogText3D, Verbose, TEXT("%s : Text3DRenderer %s Cleared"), !!Owner ? *Owner->GetActorNameOrLabel() : TEXT("Invalid owner"), *GetName().ToString())
 }
 
 void UText3DRendererBase::Destroy()
diff --git a/Engine/Plugins/Experimental/Text3D/Source/Text3D/Public/Extensions/Text3DDefaultGeometryExtension.h b/Engine/Plugins/Experimental/Text3D/Source/Text3D/Public/Extensions/Text3DDefaultGeometryExtension.h
--- a/Engine/Plugins/Experimental/Text3D/Source/Text3D/Public/Extensions/Text3DDefaultGeometryExtension.h
+++ b/Engine/Plugins/Experimental/Text3D/Source/Text3D/Public/Extensions/Text3DDefaultGeometryExtension.h
@@ -142,11 +142,11 @@ protected:
 	float Outline = 0.5f;
 
 	/** Pivot alignment of each character */
-	UPROPERTY(EditAnywhere, Getter, Setter, Category = "Geometry", meta = (AllowPrivateAccess = "true"))
+	UPROPERTY(EditAnywhere, DisplayName="Character Pivot Alignment Horizontal", Getter, Setter, Category = "Geometry", meta = (AllowPrivateAccess = "true"))
 	EText3DHorizontalTextAlignment PivotHAlignment = EText3DHorizontalTextAlignment::Left;
 	
 	/** Pivot alignment of each character */
-	UPROPERTY(EditAnywhere, Getter, Setter, Category = "Geometry", meta = (AllowPrivateAccess = "true", ValidEnumValues="Bottom"))
+	UPROPERTY(Getter, Setter, meta=(ValidEnumValues="Bottom"))
 	EText3DVerticalTextAlignment PivotVAlignment = EText3DVerticalTextAlignment::Bottom;
 
 private:
diff --git a/Engine/Plugins/Experimental/Text3D/Source/Text3D/Public/Extensions/Text3DDefaultTransformExtension.h b/Engine/Plugins/Experimental/Text3D/Source/Text3D/Public/Extensions/Text3DDefaultTransformExtension.h
--- a/Engine/Plugins/Experimental/Text3D/Source/Text3D/Public/Extensions/Text3DDefaultTransformExtension.h
+++ b/Engine/Plugins/Experimental/Text3D/Source/Text3D/Public/Extensions/Text3DDefaultTransformExtension.h
@@ -6,6 +6,8 @@
 #include "Text3DTypes.h"
 #include "Text3DDefaultTransformExtension.generated.h"
 
+class UCurveFloat;
+
 /** Extension that handles transform data for Text3D */
 UCLASS(MinimalAPI, EditInlineNew)
 class UText3DDefaultTransformExtension : public UText3DTransformExtensionBase
@@ -38,19 +40,19 @@ public:
 	}
 
 	UFUNCTION(BlueprintCallable, Category = "Text3D|Transform")
-	TEXT3D_API void SetLocationRange(float Range);
+	TEXT3D_API void SetLocationDistance(const FVector& Distance);
 
-	float GetLocationRange() const
+	FVector GetLocationDistance() const
 	{
-		return LocationRange;
+		return LocationDistance;
 	}
 
 	UFUNCTION(BlueprintCallable, Category = "Text3D|Transform")
-	TEXT3D_API void SetLocationDistance(const FVector& Distance);
+	TEXT3D_API void SetLocationEaseCurve(UCurveFloat* InEaseCurve);
 
-	FVector GetLocationDistance() const
+	UCurveFloat* GetLocationEaseCurve() const
 	{
-		return LocationDistance;
+		return LocationEaseCurve;
 	}
 
 	UFUNCTION(BlueprintCallable, Category = "Text3D|Transform")
@@ -77,14 +79,6 @@ public:
 		return RotationOrder;
 	}
 
-	UFUNCTION(BlueprintCallable, Category = "Text3D|Transform")
-	TEXT3D_API void SetRotationRange(float Range);
-
-	float GetRotationRange() const
-	{
-		return RotationRange;
-	}
-
 	UFUNCTION(BlueprintCallable, Category = "Text3D|Transform")
 	TEXT3D_API void SetRotationBegin(const FRotator& Value);
 
@@ -101,6 +95,14 @@ public:
 		return RotationEnd;
 	}
 
+	UFUNCTION(BlueprintCallable, Category = "Text3D|Transform")
+	TEXT3D_API void SetRotationEaseCurve(UCurveFloat* InEaseCurve);
+
+	UCurveFloat* GetRotationEaseCurve() const
+	{
+		return RotationEaseCurve;
+	}
+
 	UFUNCTION(BlueprintCallable, Category = "Text3D|Transform")
 	TEXT3D_API void SetScaleEnabled(bool bEnabled);
 
@@ -125,14 +127,6 @@ public:
 		return ScaleOrder;
 	}
 
-	UFUNCTION(BlueprintCallable, Category = "Text3D|Transform")
-	TEXT3D_API void SetScaleRange(float Range);
-
-	float GetScaleRange() const
-	{
-		return ScaleRange;
-	}
-
 	UFUNCTION(BlueprintCallable, Category = "Text3D|Transform")
 	TEXT3D_API void SetScaleBegin(const FVector& Value);
 
@@ -149,6 +143,15 @@ public:
 		return ScaleEnd;
 	}
 
+	UFUNCTION(BlueprintCallable, Category = "Text3D|Transform")
+	TEXT3D_API void SetScaleEaseCurve(UCurveFloat* InEaseCurve);
+
+	UCurveFloat* GetScaleEaseCurve() const
+	{
+		return ScaleEaseCurve;
+	}
+
+
 protected:
 	UText3DDefaultTransformExtension();
 
@@ -163,64 +166,67 @@ protected:
 	//~ End UText3DTransformExtensionBase
 
 	void OnTransformOptionsChanged();
-	float GetEffectValue(int32 Index, int32 Total, EText3DCharacterEffectOrder Order, float Progress, float Range);
-	void GetLineParameters(float Range, EText3DCharacterEffectOrder Order, int32 Total, float & EffectOut, float & StripOut);
-	int32 GetEffectPosition(int32 Index, int32 Total, EText3DCharacterEffectOrder Order);
+
+	int32 GetEffectPosition(int32 Index, int32 Total, EText3DCharacterEffectOrder Order) const;
+	float CalculateEffect(int32 InIndex, int32 InTotal, EText3DCharacterEffectOrder InOrder, float InProgress, const UCurveFloat* InEaseCurve) const;
 
 	// Location
 
 	UPROPERTY(EditAnywhere, Getter="GetLocationEnabled", Setter="SetLocationEnabled", Category = "Transform")
 	bool bLocationEnabled = false;
 
-	UPROPERTY(EditAnywhere, Getter, Setter, Category = "Transform", meta = (EditCondition = "bLocationEnabled", EditConditionHides, ClampMin = 0, ClampMax = 100))
+	UPROPERTY(EditAnywhere, Getter, Setter, Category = "Transform", meta = (EditCondition = "bLocationEnabled", EditConditionHides, ClampMin = 0, ClampMax = 100, Units=Percent))
 	float LocationProgress = 0.f;
 
 	UPROPERTY(EditAnywhere, Getter, Setter, Category = "Transform", meta = (EditCondition = "bLocationEnabled", EditConditionHides))
 	EText3DCharacterEffectOrder LocationOrder = EText3DCharacterEffectOrder::Normal;
 
-	UPROPERTY(EditAnywhere, Getter, Setter, Category = "Transform", meta = (EditCondition = "bLocationEnabled", EditConditionHides, ClampMin = 0, ClampMax = 100))
-	float LocationRange = 50.f;
-
 	UPROPERTY(EditAnywhere, Getter, Setter, Category = "Transform", meta = (EditCondition = "bLocationEnabled", EditConditionHides))
 	FVector LocationDistance = FVector(100.f, 0.f, 0.f);
 
+	/** Provide a 0-1 ease curve, leaving this unset will result in linear ease */
+	UPROPERTY(EditAnywhere, Setter, Getter, Category = "Transform", meta = (EditCondition = "bLocationEnabled", EditConditionHides))
+	TObjectPtr<UCurveFloat> LocationEaseCurve;
+
 	// Rotate
 
 	UPROPERTY(EditAnywhere, Getter="GetRotationEnabled", Setter="SetRotationEnabled", Category = "Transform")
 	bool bRotationEnabled = false;
 
-	UPROPERTY(EditAnywhere, Getter, Setter, Category = "Transform", meta = (EditCondition = "bRotationEnabled", EditConditionHides, ClampMin = 0, ClampMax = 100))
+	UPROPERTY(EditAnywhere, Getter, Setter, Category = "Transform", meta = (EditCondition = "bRotationEnabled", EditConditionHides, ClampMin = 0, ClampMax = 100, Units=Percent))
 	float RotationProgress = 0.0f;
 
 	UPROPERTY(EditAnywhere, Getter, Setter, Category = "Transform", meta = (EditCondition = "bRotationEnabled", EditConditionHides))
 	EText3DCharacterEffectOrder RotationOrder = EText3DCharacterEffectOrder::Normal;
 
-	UPROPERTY(EditAnywhere, Getter, Setter, Category = "Transform", meta = (EditCondition = "bRotationEnabled", EditConditionHides, ClampMin = 0, ClampMax = 100))
-	float RotationRange = 50.f;
-
 	UPROPERTY(EditAnywhere, Getter, Setter, Category = "Transform", meta = (EditCondition = "bRotationEnabled", EditConditionHides))
 	FRotator RotationBegin = FRotator(-90.f, 0.f, 0.f);
 
 	UPROPERTY(EditAnywhere, Getter, Setter, Category = "Transform", meta = (EditCondition = "bRotationEnabled", EditConditionHides))
 	FRotator RotationEnd = FRotator(0.f, 0.f, 0.f);
 
+	/** Provide a 0-1 ease curve, leaving this unset will result in linear ease */
+	UPROPERTY(EditAnywhere, Setter, Getter, Category = "Transform", meta = (EditCondition = "bRotationEnabled", EditConditionHides))
+	TObjectPtr<UCurveFloat> RotationEaseCurve;
+
 	// Scale
 
 	UPROPERTY(EditAnywhere, Getter="GetScaleEnabled", Setter="SetScaleEnabled", Category = "Transform")
 	bool bScaleEnabled = false;
 
-	UPROPERTY(EditAnywhere, Getter, Setter, Category = "Transform", meta = (EditCondition = "bScaleEnabled", EditConditionHides, ClampMin = 0, ClampMax = 100))
+	UPROPERTY(EditAnywhere, Getter, Setter, Category = "Transform", meta = (EditCondition = "bScaleEnabled", EditConditionHides, ClampMin = 0, ClampMax = 100, Units=Percent))
 	float ScaleProgress = 0.f;
 
 	UPROPERTY(EditAnywhere, Getter, Setter, Category = "Transform", meta = (EditCondition = "bScaleEnabled", EditConditionHides))
 	EText3DCharacterEffectOrder ScaleOrder = EText3DCharacterEffectOrder::Normal;
 
-	UPROPERTY(EditAnywhere, Getter, Setter, Category = "Transform", meta = (EditCondition = "bScaleEnabled", EditConditionHides, ClampMin = 0, ClampMax = 100))
-	float ScaleRange = 50.f;
-
 	UPROPERTY(EditAnywhere, Getter, Setter, Category = "Transform", meta = (EditCondition = "bScaleEnabled", EditConditionHides, ClampMin = 0))
 	FVector ScaleBegin = FVector(1.f, 0.f, 0.f);
 
 	UPROPERTY(EditAnywhere, Getter, Setter, Category = "Transform", meta = (EditCondition = "bScaleEnabled", EditConditionHides, ClampMin = 0))
 	FVector ScaleEnd = FVector::OneVector;
+
+	/** Provide a 0-1 ease curve, leaving this unset will result in linear ease */
+	UPROPERTY(EditAnywhere, Setter, Getter, Category = "Transform", meta = (EditCondition = "bScaleEnabled", EditConditionHides))
+	TObjectPtr<UCurveFloat> ScaleEaseCurve;
 };
diff --git a/Engine/Plugins/Experimental/Text3D/Source/Text3D/Public/Text3DTypes.h b/Engine/Plugins/Experimental/Text3D/Source/Text3D/Public/Text3DTypes.h
--- a/Engine/Plugins/Experimental/Text3D/Source/Text3D/Public/Text3DTypes.h
+++ b/Engine/Plugins/Experimental/Text3D/Source/Text3D/Public/Text3DTypes.h
@@ -99,10 +99,10 @@ enum class EText3DMaxWidthHandling : uint8
 UENUM()
 enum class EText3DCharacterEffectOrder : uint8
 {
-	Normal			UMETA(DisplayName = "Normal"),
+	Normal			UMETA(DisplayName = "Left To Right"),
 	FromCenter		UMETA(DisplayName = "From Center"),
 	ToCenter		UMETA(DisplayName = "To Center"),
-	Opposite		UMETA(DisplayName = "Opposite"),
+	Opposite		UMETA(DisplayName = "Right To Left"),
 };
 
 struct FText3DWordStatistics
