MotionDesign : Text3D
- Fix text character details customization firing when text geometry changes are detected causing unresponsive details panel

diff --git a/Engine/Plugins/Experimental/Text3D/Source/Text3DEditor/Private/Customizations/Text3DEditorCharacterPropertyTypeCustomization.cpp b/Engine/Plugins/Experimental/Text3D/Source/Text3DEditor/Private/Customizations/Text3DEditorCharacterPropertyTypeCustomization.cpp
--- a/Engine/Plugins/Experimental/Text3D/Source/Text3DEditor/Private/Customizations/Text3DEditorCharacterPropertyTypeCustomization.cpp
+++ b/Engine/Plugins/Experimental/Text3D/Source/Text3DEditor/Private/Customizations/Text3DEditorCharacterPropertyTypeCustomization.cpp
@@ -38,26 +38,31 @@ namespace UE::Text3DEditor::Customization
 		}
 
 		PropertyUtilitiesWeak = InUtils.GetPropertyUtilities().ToWeakPtr();
+		CharacterExtensionWeak = CharacterExtension;
+		CharacterPropertyHandleWeak = InPropertyHandle;
+
+		OnTextCharacterChanged();
+
 		if (UText3DComponent* Text3DComponent = CharacterExtension->GetText3DComponent())
 		{
-			Text3DComponent->OnTextPostUpdate().AddSPLambda(this, [this](UText3DComponent*, EText3DRendererFlags InFlags)
+			Text3DComponent->OnTextPostUpdate().AddSPLambda(this, [this](UText3DComponent* InComponent, EText3DRendererFlags InFlags)
 			{
 				if (EnumHasAnyFlags(InFlags, EText3DRendererFlags::Geometry))
 				{
-					if (const TSharedPtr<IPropertyUtilities> Utilities = PropertyUtilitiesWeak.Pin())
+					const TOptional<uint16> MaxIndex = GetCharacterLastIndex();
+					if (ActiveIndex > MaxIndex.Get(0))
 					{
-						Utilities->RequestForceRefresh();
+						OnTextCharacterChanged();
+
+						if (const TSharedPtr<IPropertyHandle> ParentProperty = CharacterPropertyHandleWeak.Pin())
+						{
+							ParentProperty->RequestRebuildChildren();
+						}
 					}
 				}
 			});
 		}
 
-		CharacterExtensionWeak = CharacterExtension;
-		CharacterPropertyHandleWeak = InPropertyHandle;
-		ActiveIndex = CharacterExtension->TextCharacterIndex;
-
-		const int32 LastCharacterIndex = FMath::Max(0, CharacterExtension->GetCharacterCount() - 1);
-
 		InHeaderRow
 		.NameContent()
 		[
@@ -82,9 +87,9 @@ namespace UE::Text3DEditor::Customization
 		]
 		.ValueContent()
 		[
-			SNew(SSpinBox<int32>)
+			SNew(SSpinBox<uint16>)
 			.MinValue(0)
-			.MaxValue(LastCharacterIndex)
+			.MaxValue(this, &FText3DEditorCharacterPropertyTypeCustomization::GetCharacterLastIndex)
 			.Delta(1)
 			.EnableSlider(true)
 			.PreventThrottling(true)
@@ -104,7 +109,6 @@ namespace UE::Text3DEditor::Customization
 					if (InType != ETextCommit::Type::OnCleared)
 					{
 						ActiveIndex = InNewValue;
-
 						Extension->TextCharacterIndex = ActiveIndex;
 
 						if (const TSharedPtr<IPropertyUtilities> Utilities = PropertyUtilitiesWeak.Pin())
@@ -164,7 +168,7 @@ namespace UE::Text3DEditor::Customization
 	{
 		if (const UText3DCharacterExtensionBase* Extension = CharacterExtensionWeak.Get())
 		{
-			const UText3DCharacterBase* Character = Extension->GetCharacter(ActiveIndex);
+			const UText3DCharacterBase* Character = Extension->GetCharacter(Extension->TextCharacterIndex);
 
 			return FText::Format(
 				LOCTEXT("CharacterCount", "Total: {0} / Current: {1}")
@@ -175,6 +179,25 @@ namespace UE::Text3DEditor::Customization
 
 		return FText::GetEmpty();
 	}
+
+	TOptional<uint16> FText3DEditorCharacterPropertyTypeCustomization::GetCharacterLastIndex() const
+	{
+		if (const UText3DCharacterExtensionBase* Extension = CharacterExtensionWeak.Get())
+		{
+			return FMath::Max(0, Extension->GetCharacterCount() - 1);
+		}
+
+		return 0;
+	}
+
+	void FText3DEditorCharacterPropertyTypeCustomization::OnTextCharacterChanged()
+	{
+		if (UText3DCharacterExtensionBase* Extension = CharacterExtensionWeak.Get())
+		{
+			ActiveIndex = FMath::Clamp(Extension->TextCharacterIndex, 0, GetCharacterLastIndex().Get(0));
+			Extension->TextCharacterIndex = ActiveIndex;
+		}
+	}
 }
 
 #undef LOCTEXT_NAMESPACE
diff --git a/Engine/Plugins/Experimental/Text3D/Source/Text3DEditor/Private/Customizations/Text3DEditorCharacterPropertyTypeCustomization.h b/Engine/Plugins/Experimental/Text3D/Source/Text3DEditor/Private/Customizations/Text3DEditorCharacterPropertyTypeCustomization.h
--- a/Engine/Plugins/Experimental/Text3D/Source/Text3DEditor/Private/Customizations/Text3DEditorCharacterPropertyTypeCustomization.h
+++ b/Engine/Plugins/Experimental/Text3D/Source/Text3DEditor/Private/Customizations/Text3DEditorCharacterPropertyTypeCustomization.h
@@ -3,6 +3,7 @@
 #pragma once
 
 #include "IPropertyTypeCustomization.h"
+#include "Misc/Optional.h"
 #include "PropertyEditorModule.h"
 #include "PropertyHandle.h"
 #include "Templates/SharedPointer.h"
@@ -45,8 +46,10 @@ namespace UE::Text3DEditor::Customization
 
 	private:
 		FText GetCharacterText() const;
+		TOptional<uint16> GetCharacterLastIndex() const;
+		void OnTextCharacterChanged();
 
-		int32 ActiveIndex = 0;
+		uint16 ActiveIndex = 0;
 		TWeakPtr<IPropertyHandle> CharacterPropertyHandleWeak;
 		TWeakObjectPtr<UText3DCharacterExtensionBase> CharacterExtensionWeak;
 		TWeakPtr<IPropertyUtilities> PropertyUtilitiesWeak;
