MotionDesign : Text3D
- Adjusting default for rendering extension to match other class engine behavior

diff --git a/Engine/Plugins/Experimental/Text3D/Source/Text3D/Public/Extensions/Text3DDefaultRenderingExtension.h b/Engine/Plugins/Experimental/Text3D/Source/Text3D/Public/Extensions/Text3DDefaultRenderingExtension.h
--- a/Engine/Plugins/Experimental/Text3D/Source/Text3D/Public/Extensions/Text3DDefaultRenderingExtension.h
+++ b/Engine/Plugins/Experimental/Text3D/Source/Text3D/Public/Extensions/Text3DDefaultRenderingExtension.h
@@ -95,7 +95,7 @@ protected:
 
 	/** Controls whether the primitive should influence indirect lighting. */
 	UPROPERTY(EditAnywhere, Getter = "GetAffectDynamicIndirectLighting", Setter = "SetAffectDynamicIndirectLighting", Category = "Rendering", meta = (AllowPrivateAccess = "true"))
-	bool bAffectDynamicIndirectLighting = false;
+	bool bAffectDynamicIndirectLighting = true;
 
 	/** Controls whether the primitive should affect indirect lighting when hidden. This flag is only used if bAffectDynamicIndirectLighting is true. */
 	UPROPERTY(EditAnywhere, Getter="GetAffectIndirectLightingWhileHidden", Setter="SetAffectIndirectLightingWhileHidden", Category = "Rendering", meta = (EditCondition = "bAffectDynamicIndirectLighting", AllowPrivateAccess = "true"))
