MotionDesign : Text3D
- Fix last character tracking/kerning is skipped
- Adding base effect class
- Renaming TransformExtension as LayoutTransformEffect for future effects (geometry, material)
- Exposing material extension properties to material designer
- Adding LocationBegin in LayoutTransformEffect to adjust starting offset
- Add toggle for material mask, off by default
- Fix mask material function to work fully in local space without continous updates to adjust itself
- Fix visualizer character kerning handle only show when kerning editing is toggled on
- Adding missing blueprint functions to retrieve text extensions
- Use single custom material for front, back, extrude, bevel slots

#virtualized

diff --git a/Engine/Plugins/Experimental/Avalanche/Source/AvalancheText/Private/AvaText3DComponent.cpp b/Engine/Plugins/Experimental/Avalanche/Source/AvalancheText/Private/AvaText3DComponent.cpp
--- a/Engine/Plugins/Experimental/Avalanche/Source/AvalancheText/Private/AvaText3DComponent.cpp
+++ b/Engine/Plugins/Experimental/Avalanche/Source/AvalancheText/Private/AvaText3DComponent.cpp
@@ -6,7 +6,7 @@
 #include "Characters/Text3DCharacterTransform.h"
 #include "Engine/Texture2D.h"
 #include "Extensions/Text3DDefaultMaterialExtension.h"
-#include "Extensions/Text3DDefaultTransformExtension.h"
+#include "Extensions/Text3DLayoutTransformEffect.h"
 #include "Extensions/Text3DMaterialExtensionBase.h"
 #include "GameFramework/Actor.h"
 #include "Materials/Material.h"
@@ -119,12 +119,12 @@ void UAvaText3DComponent::Serialize(FArchive& InArchive)
 			{
 				if (const UText3DCharacterTransform* CharacterTransform = Actor->FindComponentByClass<UText3DCharacterTransform>())
 				{
-					UText3DDefaultTransformExtension* TransformExtension = NewObject<UText3DDefaultTransformExtension>(this, TEXT("AvaTransformExtension"));
+					UText3DLayoutTransformEffect* TransformExtension = NewObject<UText3DLayoutTransformEffect>(this, TEXT("AvaTransformExtension"));
 
 					TransformExtension->SetLocationEnabled(CharacterTransform->GetLocationEnabled());
 					TransformExtension->SetLocationProgress(CharacterTransform->GetLocationProgress());
 					TransformExtension->SetLocationOrder(CharacterTransform->GetLocationOrder());
-					TransformExtension->SetLocationDistance(CharacterTransform->GetLocationDistance());
+					TransformExtension->SetLocationEnd(CharacterTransform->GetLocationDistance());
 
 					TransformExtension->SetRotationEnabled(CharacterTransform->GetRotationEnabled());
 					TransformExtension->SetRotationProgress(CharacterTransform->GetRotationProgress());
@@ -138,7 +138,7 @@ void UAvaText3DComponent::Serialize(FArchive& InArchive)
 					TransformExtension->SetScaleBegin(CharacterTransform->GetScaleBegin());
 					TransformExtension->SetScaleEnd(CharacterTransform->GetScaleEnd());
 
-					TransformExtensions.Emplace(TransformExtension);
+					LayoutEffects.Emplace(TransformExtension);
 				}
 			}
 
diff --git a/Engine/Plugins/Experimental/Avalanche/Source/AvalancheTextEditor/Private/AvaTextEditorModule.cpp b/Engine/Plugins/Experimental/Avalanche/Source/AvalancheTextEditor/Private/AvaTextEditorModule.cpp
--- a/Engine/Plugins/Experimental/Avalanche/Source/AvalancheTextEditor/Private/AvaTextEditorModule.cpp
+++ b/Engine/Plugins/Experimental/Avalanche/Source/AvalancheTextEditor/Private/AvaTextEditorModule.cpp
@@ -10,12 +10,14 @@
 #include "ColorPicker/AvaViewportColorPickerActorClassRegistry.h"
 #include "ColorPicker/AvaViewportColorPickerAdapter.h"
 #include "DMObjectMaterialProperty.h"
+#include "Extensions/Text3DDefaultMaterialExtension.h"
 #include "Framework/Application/SlateApplication.h"
 #include "IAvalancheComponentVisualizersModule.h"
 #include "IDynamicMaterialEditorModule.h"
 #include "Misc/CoreDelegates.h"
 #include "Modules/ModuleManager.h"
 #include "PropertyEditorModule.h"
+#include "Text3DActor.h"
 #include "Visualizer/AvaTextVisualizer.h"
 
 void FAvaTextEditorModule::StartupModule()
@@ -48,40 +50,57 @@ void FAvaTextEditorModule::RegisterComponentVisualizers()
 
 void FAvaTextEditorModule::RegisterDynamicMaterialPropertyGenerator()
 {
-	IDynamicMaterialEditorModule::Get().RegisterCustomMaterialPropertyGenerator(
+	const TArray<UClass*> RegisterClasses
+	{
+		AText3DActor::StaticClass(),
 		AAvaTextActor::StaticClass(),
-		FDMGetObjectMaterialPropertiesDelegate::CreateLambda(
-			[](UObject* InObject)
-			{
-				static TArray<FProperty*> Text3DProperties = {
-					UText3DComponent::StaticClass()->FindPropertyByName("FrontMaterial"),
-					UText3DComponent::StaticClass()->FindPropertyByName("BevelMaterial"),
-					UText3DComponent::StaticClass()->FindPropertyByName("ExtrudeMaterial"),
-					UText3DComponent::StaticClass()->FindPropertyByName("BackMaterial")
-				};
-
-				TArray<FDMObjectMaterialProperty> Properties;
-
-				if (AAvaTextActor* AvaTextActor = Cast<AAvaTextActor>(InObject))
+	};
+
+	/**
+	 * Customize Text actor to only show material extension materials instead of all character slots
+	 */
+	for (UClass* RegisterClass : RegisterClasses)
+	{
+		IDynamicMaterialEditorModule::Get().RegisterCustomMaterialPropertyGenerator(
+			RegisterClass,
+			FDMGetObjectMaterialPropertiesDelegate::CreateLambda(
+				[](UObject* InObject)
 				{
-					if (UText3DComponent* TextComponent = AvaTextActor->GetText3DComponent())
+					UClass* MaterialExtensionClass = UText3DDefaultMaterialExtension::StaticClass();
+					static TArray<FProperty*> Text3DProperties =
 					{
-						Properties.Reserve(Text3DProperties.Num());
+						MaterialExtensionClass->FindPropertyByName("FrontMaterial"),
+						MaterialExtensionClass->FindPropertyByName("BevelMaterial"),
+						MaterialExtensionClass->FindPropertyByName("ExtrudeMaterial"),
+						MaterialExtensionClass->FindPropertyByName("BackMaterial")
+					};
+
+					TArray<FDMObjectMaterialProperty> Properties;
 
-						for (FProperty* Property : Text3DProperties)
+					if (const AActor* TextActor = Cast<AActor>(InObject))
+					{
+						if (const UText3DComponent* TextComponent = TextActor->FindComponentByClass<UText3DComponent>())
 						{
-							if (Property)
+							if (UText3DDefaultMaterialExtension* MaterialExtension = TextComponent->GetCastedMaterialExtension<UText3DDefaultMaterialExtension>())
 							{
-								Properties.Emplace(TextComponent, Property);
+								Properties.Reserve(Text3DProperties.Num());
+
+								for (FProperty* Property : Text3DProperties)
+								{
+									if (Property)
+									{
+										Properties.Emplace(MaterialExtension, Property);
+									}
+								}
 							}
 						}
 					}
-				}
 
-				return Properties;
-			}
-		)
-	);
+					return Properties;
+				}
+			)
+		);
+	}
 }
 
 IMPLEMENT_MODULE(FAvaTextEditorModule, AvalancheTextEditor)
diff --git a/Engine/Plugins/Experimental/Avalanche/Source/AvalancheTextEditor/Private/Visualizer/AvaTextVisualizer.cpp b/Engine/Plugins/Experimental/Avalanche/Source/AvalancheTextEditor/Private/Visualizer/AvaTextVisualizer.cpp
--- a/Engine/Plugins/Experimental/Avalanche/Source/AvalancheTextEditor/Private/Visualizer/AvaTextVisualizer.cpp
+++ b/Engine/Plugins/Experimental/Avalanche/Source/AvalancheTextEditor/Private/Visualizer/AvaTextVisualizer.cpp
@@ -58,6 +58,7 @@ FAvaTextVisualizer::FAvaTextVisualizer()
 
 	ResetEditingFlags();
 	bShowGradientControls = false;
+	bEditingKerning = false;
 }
 
 UActorComponent* FAvaTextVisualizer::GetEditedComponent() const
@@ -291,13 +292,22 @@ bool FAvaTextVisualizer::VisProxyHandleClick(FEditorViewportClient* InViewportCl
 		{
 			EndEditing();
 
-			if (UText3DComponent* Text3DComponent = Cast<UText3DComponent>(const_cast<UActorComponent*>(KerningHandleProxy->Component.Get())))
+			if (KerningHandleProxy->Index != std::numeric_limits<uint16>::max())
 			{
-				TextComponent = Text3DComponent;
-				EditingKerningIndex = KerningHandleProxy->Index;
-				StartEditing(InViewportClient, Text3DComponent);
-				return true;
-			}			
+				if (UText3DComponent* Text3DComponent = Cast<UText3DComponent>(const_cast<UActorComponent*>(KerningHandleProxy->Component.Get())))
+				{
+					TextComponent = Text3DComponent;
+					EditingKerningIndex = KerningHandleProxy->Index;
+					StartEditing(InViewportClient, Text3DComponent);
+				}
+			}
+			else
+			{
+				bEditingKerning = !bEditingKerning;
+				EditingKerningIndex = INDEX_NONE;
+			}
+
+			return true;
 		}
 	}
 
@@ -1036,6 +1046,7 @@ void FAvaTextVisualizer::DrawVisualizationNotEditing(const UActorComponent* InCo
 		if (MaterialExtension->GetStyle() == EText3DMaterialStyle::Gradient)
 		{
 			DrawEditGradientButton(Text3DComponent, InView, InPDI, InOutIconIndex, bShowGradientControls ? Enabled : Disabled);
+			++InOutIconIndex;
 
 			if (bShowGradientControls)
 			{
@@ -1044,11 +1055,17 @@ void FAvaTextVisualizer::DrawVisualizationNotEditing(const UActorComponent* InCo
 		}
 	}
 
-	Text3DComponent->ForEachCharacter([this, Text3DComponent, InView, InPDI](UText3DCharacterBase* InCharacter, uint16 InIndex, uint16)
+	DrawCharacterKerningButton(Text3DComponent, InView, InPDI, InOutIconIndex, bEditingKerning ? Enabled : Disabled);
+	++InOutIconIndex;
+
+	if (bEditingKerning)
 	{
-		const FVector CharacterLocation = GetKerningHandleLocation(Text3DComponent, InIndex);
-		DrawCharacterKerningHandle(Text3DComponent, InView, InPDI, CharacterLocation, InIndex, Inactive);
-	});
+		Text3DComponent->ForEachCharacter([this, Text3DComponent, InView, InPDI](UText3DCharacterBase* InCharacter, uint16 InIndex, uint16)
+		{
+			const FVector CharacterLocation = GetKerningHandleLocation(Text3DComponent, InIndex);
+			DrawCharacterKerningHandle(Text3DComponent, InView, InPDI, CharacterLocation, InIndex, Inactive);
+		});
+	}
 }
 
 void FAvaTextVisualizer::DrawVisualizationEditing(const UActorComponent* InComponent, const FSceneView* InView, 
@@ -1104,6 +1121,9 @@ void FAvaTextVisualizer::DrawVisualizationEditing(const UActorComponent* InCompo
 		}
 	}
 
+	DrawCharacterKerningButton(Text3DComponent, InView, InPDI, InOutIconIndex, bEditingKerning ? Enabled : Disabled);
+	++InOutIconIndex;
+
 	if (EditingKerningIndex != INDEX_NONE)
 	{
 		Text3DComponent->ForEachCharacter([this, Text3DComponent, InView, InPDI](UText3DCharacterBase* InCharacter, uint16 InIndex, uint16)
@@ -1513,6 +1533,31 @@ void FAvaTextVisualizer::DrawGradientLineEndHandle(const UText3DComponent* InTex
 	InPDI->SetHitProxy(nullptr);
 }
 
+void FAvaTextVisualizer::DrawCharacterKerningButton(const UText3DComponent* InTextComponent, const FSceneView* InView, FPrimitiveDrawInterface* InPDI,
+	int32 InIconIndex, const FLinearColor& InColor)
+{
+	if (!InTextComponent)
+	{
+		return;
+	}
+
+	UTexture2D* KerningSprite = IAvalancheComponentVisualizersModule::Get().GetSettings()->GetVisualizerSprite(UE::AvaShapes::InnerSizeSprite);
+
+	if (!KerningSprite || !KerningSprite->GetResource())
+	{
+		return;
+	}
+
+	FVector IconLocation;
+	float IconSize;
+	GetIconMetrics(InView, InIconIndex, IconLocation, IconSize);
+
+	InPDI->SetHitProxy(new HAvaTextCharacterKerningHandleProxy(InTextComponent, std::numeric_limits<uint16>::max()));
+	InPDI->DrawSprite(IconLocation, IconSize, IconSize, KerningSprite->GetResource(), InColor,
+		SDPG_Foreground, 0, 0, 0, 0, SE_BLEND_Opaque);
+	InPDI->SetHitProxy(nullptr);
+}
+
 void FAvaTextVisualizer::DrawCharacterKerningHandle(const UText3DComponent* InTextComponent, const FSceneView* InView, FPrimitiveDrawInterface* InPDI,
 	const FVector& InLocation, uint16 InCharacterIndex, const FLinearColor& InColor) const
 {
@@ -1523,9 +1568,9 @@ void FAvaTextVisualizer::DrawCharacterKerningHandle(const UText3DComponent* InTe
 		return;
 	}
 
-	const UTexture2D* const UVSprite = IAvalancheComponentVisualizersModule::Get().GetSettings()->GetVisualizerSprite(UE::AvaShapes::SizeSprite);
+	const UTexture2D* const KerningSprite = IAvalancheComponentVisualizersModule::Get().GetSettings()->GetVisualizerSprite(UE::AvaShapes::SizeSprite);
 
-	if (!UVSprite || !UVSprite->GetResource())
+	if (!KerningSprite || !KerningSprite->GetResource())
 	{
 		return;
 	}
@@ -1534,7 +1579,7 @@ void FAvaTextVisualizer::DrawCharacterKerningHandle(const UText3DComponent* InTe
 	const float IconSize = BaseSize * GetIconSizeScale(InView, IconLocation);
 
 	InPDI->SetHitProxy(new HAvaTextCharacterKerningHandleProxy(InTextComponent, InCharacterIndex));
-	InPDI->DrawSprite(IconLocation, IconSize, IconSize, UVSprite->GetResource(), InColor,
+	InPDI->DrawSprite(IconLocation, IconSize, IconSize, KerningSprite->GetResource(), InColor,
 		SDPG_Foreground, 0, 0, 0, 0, SE_BLEND_Opaque);
 	InPDI->SetHitProxy(nullptr);
 }
diff --git a/Engine/Plugins/Experimental/Avalanche/Source/AvalancheTextEditor/Private/Visualizer/AvaTextVisualizer.h b/Engine/Plugins/Experimental/Avalanche/Source/AvalancheTextEditor/Private/Visualizer/AvaTextVisualizer.h
--- a/Engine/Plugins/Experimental/Avalanche/Source/AvalancheTextEditor/Private/Visualizer/AvaTextVisualizer.h
+++ b/Engine/Plugins/Experimental/Avalanche/Source/AvalancheTextEditor/Private/Visualizer/AvaTextVisualizer.h
@@ -163,6 +163,7 @@ protected:
 	bool bEditingGradientRotation_EndHandle;
 	bool bEditingGradientOffset;
 	bool bEditingGradientSmoothness;
+	bool bEditingKerning;
 	int32 EditingKerningIndex = INDEX_NONE;
 	
 	float InitialGradientRotation;
@@ -236,7 +237,10 @@ protected:
 
 	void DrawGradientLineEndHandle(const UText3DComponent* InTextComponent, const FSceneView* InView,
 		FPrimitiveDrawInterface* InPDI, const FVector& InLocation, const FLinearColor& InColor) const;
-	
+
+	void DrawCharacterKerningButton(const UText3DComponent* InTextComponent, const FSceneView* InView, FPrimitiveDrawInterface* InPDI,
+		int32 InIconIndex, const FLinearColor& InColor);
+
 	void DrawCharacterKerningHandle(const UText3DComponent* InTextComponent, const FSceneView* InView,
 		FPrimitiveDrawInterface* InPDI, const FVector& InLocation, uint16 InCharacterIndex, const FLinearColor& InColor) const;
 };
diff --git a/Engine/Plugins/Experimental/Text3D/Source/Text3D/Private/Extensions/Text3DDefaultLayoutExtension.cpp b/Engine/Plugins/Experimental/Text3D/Source/Text3D/Private/Extensions/Text3DDefaultLayoutExtension.cpp
--- a/Engine/Plugins/Experimental/Text3D/Source/Text3D/Private/Extensions/Text3DDefaultLayoutExtension.cpp
+++ b/Engine/Plugins/Experimental/Text3D/Source/Text3D/Private/Extensions/Text3DDefaultLayoutExtension.cpp
@@ -6,7 +6,7 @@
 #include "Characters/Text3DCharacterBase.h"
 #include "Engine/Font.h"
 #include "Extensions/Text3DGeometryExtensionBase.h"
-#include "Extensions/Text3DTransformExtensionBase.h"
+#include "Extensions/Text3DLayoutEffectBase.h"
 #include "Fonts/SlateTextShaper.h"
 #include "Framework/Text/PlainTextLayoutMarshaller.h"
 #include "LayoutBuilders/Text3DLayout.h"
@@ -286,8 +286,6 @@ EText3DExtensionResult UText3DDefaultLayoutExtension::PreRendererUpdate(EText3DR
 	ShapedText->CalculateWidth();
 	CalculateTextScale();
 
-	const TConstArrayView<UText3DTransformExtensionBase*> TransformExtensions = Text3DComponent->GetTransformExtensions();
-
 	int32 LayoutIndex = 0;
 	for (TConstEnumerateRef<FShapedGlyphLine> ShapedLine : EnumerateRange(ShapedText->Lines))
 	{
@@ -317,20 +315,12 @@ EText3DExtensionResult UText3DDefaultLayoutExtension::PreRendererUpdate(EText3DR
 					break;
 				}
 
-				if (ShapedGlyph.GetIndex() != 0 && ShapedGlyph.GetIndex() != ShapedLine->GlyphsToRender.Num() - 1)
+				if (ShapedGlyph.GetIndex() != 0)
 				{
 					Shift.Y += Tracking + ShapedText->Kernings[*CharacterIndex];
 				}
 
 				CharacterTransform.Accumulate(FTransform(Location + Shift));
-
-				for (UText3DTransformExtensionBase* TransformExtension : TransformExtensions)
-				{
-					if (TransformExtension && TransformExtension->GetTargetRange().IsInRange(*CharacterIndex))
-					{
-						TransformExtension->ApplyTransform(*CharacterIndex, ShapedText->CharacterGlyphMap.Num(), CharacterTransform);
-					}
-				}
 			}
 
 			Location.Y += ShapedLine->GlyphAdvances[ShapedGlyph.GetIndex()];
diff --git a/Engine/Plugins/Experimental/Text3D/Source/Text3D/Private/Extensions/Text3DDefaultMaterialExtension.cpp b/Engine/Plugins/Experimental/Text3D/Source/Text3D/Private/Extensions/Text3DDefaultMaterialExtension.cpp
--- a/Engine/Plugins/Experimental/Text3D/Source/Text3D/Private/Extensions/Text3DDefaultMaterialExtension.cpp
+++ b/Engine/Plugins/Experimental/Text3D/Source/Text3D/Private/Extensions/Text3DDefaultMaterialExtension.cpp
@@ -20,7 +20,7 @@ void UText3DDefaultMaterialExtension::SetStyle(EText3DMaterialStyle InStyle)
 	}
 
 	Style = InStyle;
-	OnMaterialOptionsChanged();
+	OnCustomMaterialChanged();
 }
 
 void UText3DDefaultMaterialExtension::SetFrontColor(const FLinearColor& InColor)
@@ -177,6 +177,17 @@ void UText3DDefaultMaterialExtension::SetOpacity(float InOpacity)
 	OnMaterialOptionsChanged();
 }
 
+void UText3DDefaultMaterialExtension::SetUseMask(bool bInUseMask)
+{
+	if (bUseMask == bInUseMask)
+	{
+		return;
+	}
+
+	bUseMask = bInUseMask;
+	OnMaterialOptionsChanged();
+}
+
 void UText3DDefaultMaterialExtension::SetMaskOffset(float InMaskOffset)
 {
 	if (FMath::IsNearlyEqual(MaskOffset, InMaskOffset))
@@ -210,6 +221,17 @@ void UText3DDefaultMaterialExtension::SetMaskRotation(float InMaskRotation)
 	OnMaterialOptionsChanged();
 }
 
+void UText3DDefaultMaterialExtension::SetUseSingleMaterial(bool bInUse)
+{
+	if (bUseSingleMaterial == bInUse)
+	{
+		return;
+	}
+
+	bUseSingleMaterial = bInUse;
+	OnCustomMaterialChanged();
+}
+
 void UText3DDefaultMaterialExtension::SetFrontMaterial(UMaterialInterface* InMaterial)
 {
 	if (FrontMaterial == InMaterial)
@@ -218,7 +240,7 @@ void UText3DDefaultMaterialExtension::SetFrontMaterial(UMaterialInterface* InMat
 	}
 
 	FrontMaterial = InMaterial;
-	OnMaterialOptionsChanged();
+	OnCustomMaterialChanged();
 }
 
 void UText3DDefaultMaterialExtension::SetBevelMaterial(UMaterialInterface* InMaterial)
@@ -229,7 +251,7 @@ void UText3DDefaultMaterialExtension::SetBevelMaterial(UMaterialInterface* InMat
 	}
 	
 	BevelMaterial = InMaterial;
-	OnMaterialOptionsChanged();
+	OnCustomMaterialChanged();
 }
 
 void UText3DDefaultMaterialExtension::SetExtrudeMaterial(UMaterialInterface* InMaterial)
@@ -240,7 +262,7 @@ void UText3DDefaultMaterialExtension::SetExtrudeMaterial(UMaterialInterface* InM
 	}
 	
 	ExtrudeMaterial	= InMaterial;
-	OnMaterialOptionsChanged();
+	OnCustomMaterialChanged();
 }
 
 void UText3DDefaultMaterialExtension::SetBackMaterial(UMaterialInterface* InMaterial)
@@ -251,7 +273,7 @@ void UText3DDefaultMaterialExtension::SetBackMaterial(UMaterialInterface* InMate
 	}
 	
 	BackMaterial = InMaterial;
-	OnMaterialOptionsChanged();
+	OnCustomMaterialChanged();
 }
 
 EText3DExtensionResult UText3DDefaultMaterialExtension::PreRendererUpdate(EText3DRendererFlags InFlag)
@@ -317,23 +339,22 @@ EText3DExtensionResult UText3DDefaultMaterialExtension::PreRendererUpdate(EText3
 		const UText3DComponent* Text3DComponent = GetText3DComponent();
 
 		const FBox LocalBounds = Text3DComponent->GetBounds();
-		const FVector LocalBoundsLowerLeftCorner = Text3DComponent->GetRelativeRotation().RotateVector(LocalBounds.Min);
-		const FVector LocalBoundsExtent = Text3DComponent->GetRelativeRotation().RotateVector(LocalBounds.GetExtent());
+		const FVector LocalBoundsExtent = Text3DComponent->GetRelativeRotation().UnrotateVector(LocalBounds.GetSize());
 
 		const UText3DLayoutExtensionBase* LayoutExtension = Text3DComponent->GetLayoutExtension();
+		const FVector LineLocation = LayoutExtension->GetLineLocation(0);
 
 		const FVector TextScaleFactor = LayoutExtension->GetTextScale() * Text3DComponent->GetComponentScale();
-		const FVector BoundsLowerLeftCornerLocalScaled = LocalBoundsLowerLeftCorner / TextScaleFactor;
 		const FVector BoundsSizeScaled = LocalBoundsExtent / TextScaleFactor;
 
-		SetVectorParameter(Materials, UText3DProjectSettings::FMaterialParameters::BoundsOrigin, BoundsLowerLeftCornerLocalScaled);
+		SetVectorParameter(Materials, UText3DProjectSettings::FMaterialParameters::BoundsOrigin, LineLocation);
 		SetVectorParameter(Materials, UText3DProjectSettings::FMaterialParameters::BoundsSize, BoundsSizeScaled);
-		SetVectorParameter(Materials, UText3DProjectSettings::FMaterialParameters::TextPosition, FVector::ZeroVector);
 	}
 
 	if (BlendMode == EText3DMaterialBlendMode::Translucent)
 	{
 		SetScalarParameter(Materials, UText3DProjectSettings::FMaterialParameters::Opacity, Opacity);
+		SetScalarParameter(Materials, UText3DProjectSettings::FMaterialParameters::MaskEnabled, bUseMask ? 1 : 0);
 		SetScalarParameter(Materials, UText3DProjectSettings::FMaterialParameters::MaskOffset, MaskOffset);
 		SetScalarParameter(Materials, UText3DProjectSettings::FMaterialParameters::MaskRotation, MaskRotation);
 		SetScalarParameter(Materials, UText3DProjectSettings::FMaterialParameters::MaskSmoothness, MaskSmoothness);
@@ -443,6 +464,18 @@ void UText3DDefaultMaterialExtension::OnMaterialOptionsChanged()
 	RequestUpdate(EText3DRendererFlags::Material);
 }
 
+void UText3DDefaultMaterialExtension::OnCustomMaterialChanged()
+{
+	if (bUseSingleMaterial && Style == EText3DMaterialStyle::Custom)
+	{
+		BackMaterial = FrontMaterial;
+		BevelMaterial = FrontMaterial;
+		ExtrudeMaterial = FrontMaterial;
+	}
+
+	OnMaterialOptionsChanged();
+}
+
 UMaterialInterface* UText3DDefaultMaterialExtension::GetMaterial(EText3DGroupType InGroup) const
 {
 	switch (InGroup)
@@ -522,7 +555,6 @@ void UText3DDefaultMaterialExtension::PostEditChangeProperty(FPropertyChangedEve
 
 	static const TSet<FName> PropertyNames =
 	{
-		GET_MEMBER_NAME_CHECKED(UText3DDefaultMaterialExtension, Style),
 		GET_MEMBER_NAME_CHECKED(UText3DDefaultMaterialExtension, FrontColor),
 		GET_MEMBER_NAME_CHECKED(UText3DDefaultMaterialExtension, BackColor),
 		GET_MEMBER_NAME_CHECKED(UText3DDefaultMaterialExtension, ExtrudeColor),
@@ -537,16 +569,29 @@ void UText3DDefaultMaterialExtension::PostEditChangeProperty(FPropertyChangedEve
 		GET_MEMBER_NAME_CHECKED(UText3DDefaultMaterialExtension, BlendMode),
 		GET_MEMBER_NAME_CHECKED(UText3DDefaultMaterialExtension, bIsUnlit),
 		GET_MEMBER_NAME_CHECKED(UText3DDefaultMaterialExtension, Opacity),
+		GET_MEMBER_NAME_CHECKED(UText3DDefaultMaterialExtension, bUseMask),
 		GET_MEMBER_NAME_CHECKED(UText3DDefaultMaterialExtension, MaskOffset),
 		GET_MEMBER_NAME_CHECKED(UText3DDefaultMaterialExtension, MaskSmoothness),
 		GET_MEMBER_NAME_CHECKED(UText3DDefaultMaterialExtension, MaskRotation),
+	};
+
+	static const TSet<FName> CustomPropertyNames =
+	{
+		GET_MEMBER_NAME_CHECKED(UText3DDefaultMaterialExtension, Style),
+		GET_MEMBER_NAME_CHECKED(UText3DDefaultMaterialExtension, bUseSingleMaterial),
 		GET_MEMBER_NAME_CHECKED(UText3DDefaultMaterialExtension, FrontMaterial),
 		GET_MEMBER_NAME_CHECKED(UText3DDefaultMaterialExtension, BevelMaterial),
 		GET_MEMBER_NAME_CHECKED(UText3DDefaultMaterialExtension, ExtrudeMaterial),
 		GET_MEMBER_NAME_CHECKED(UText3DDefaultMaterialExtension, BackMaterial),
 	};
 
-	if (PropertyNames.Contains(InEvent.GetMemberPropertyName()))
+	const FName MemberPropertyName = InEvent.GetMemberPropertyName();
+
+	if (CustomPropertyNames.Contains(MemberPropertyName))
+	{
+		OnCustomMaterialChanged();
+	}
+	else if (PropertyNames.Contains(MemberPropertyName))
 	{
 		OnMaterialOptionsChanged();
 	}
diff --git a/Engine/Plugins/Experimental/Text3D/Source/Text3D/Private/Extensions/Text3DDefaultTransformExtension.cpp b/Engine/Plugins/Experimental/Text3D/Source/Text3D/Private/Extensions/Text3DDefaultTransformExtension.cpp
deleted file mode 100644
diff --git a/Engine/Plugins/Experimental/Text3D/Source/Text3D/Private/Extensions/Text3DEffectExtensionBase.cpp b/Engine/Plugins/Experimental/Text3D/Source/Text3D/Private/Extensions/Text3DEffectExtensionBase.cpp
new file mode 100644
--- /dev/null
+++ b/Engine/Plugins/Experimental/Text3D/Source/Text3D/Private/Extensions/Text3DEffectExtensionBase.cpp
@@ -0,0 +1,30 @@
+// Copyright Epic Games, Inc. All Rights Reserved.
+
+#include "Extensions/Text3DEffectExtensionBase.h"
+#include "Text3DComponent.h"
+
+EText3DExtensionResult UText3DEffectExtensionBase::PreRendererUpdate(EText3DRendererFlags InFlag)
+{
+	if (InFlag != EText3DRendererFlags::Layout)
+	{
+		return EText3DExtensionResult::Active;
+	}
+
+	const UText3DComponent* Text3DComponent = GetText3DComponent();
+
+	const uint16 CharacterCount = Text3DComponent->GetCharacterCount();
+	for (uint16 Index = 0; Index < CharacterCount; Index++)
+	{
+		if (GetTargetRange().IsInRange(Index))
+		{
+			ApplyEffect(Index, CharacterCount);
+		}
+	}
+
+	return EText3DExtensionResult::Finished;
+}
+
+EText3DExtensionResult UText3DEffectExtensionBase::PostRendererUpdate(EText3DRendererFlags InFlag)
+{
+	return EText3DExtensionResult::Active;
+}
\ No newline at end of file
diff --git a/Engine/Plugins/Experimental/Text3D/Source/Text3D/Private/Extensions/Text3DLayoutTransformEffect.cpp b/Engine/Plugins/Experimental/Text3D/Source/Text3D/Private/Extensions/Text3DLayoutTransformEffect.cpp
new file mode 100644
--- /dev/null
+++ b/Engine/Plugins/Experimental/Text3D/Source/Text3D/Private/Extensions/Text3DLayoutTransformEffect.cpp
@@ -0,0 +1,320 @@
+// Copyright Epic Games, Inc. All Rights Reserved.
+
+#include "Extensions/Text3DLayoutTransformEffect.h"
+
+#include "Characters/Text3DCharacterBase.h"
+#include "Curves/CurveFloat.h"
+#include "Text3DComponent.h"
+
+void UText3DLayoutTransformEffect::SetLocationEnabled(bool bEnabled)
+{
+	if (bLocationEnabled != bEnabled)
+	{
+		bLocationEnabled = bEnabled;
+		OnTransformOptionsChanged();
+	}
+}
+
+void UText3DLayoutTransformEffect::SetLocationProgress(float Progress)
+{
+	Progress = FMath::Clamp(Progress, 0.0f, 100.0f);
+	if (!FMath::IsNearlyEqual(LocationProgress, Progress))
+	{
+		LocationProgress = Progress;
+		OnTransformOptionsChanged();
+	}
+}
+
+void UText3DLayoutTransformEffect::SetLocationOrder(EText3DCharacterEffectOrder Order)
+{
+	if (LocationOrder != Order)
+	{
+		LocationOrder = Order;
+		OnTransformOptionsChanged();
+	}
+}
+
+void UText3DLayoutTransformEffect::SetLocationBegin(const FVector& InBegin)
+{
+	if (!LocationBegin.Equals(InBegin))
+	{
+		LocationBegin = InBegin;
+		OnTransformOptionsChanged();
+	}
+}
+
+void UText3DLayoutTransformEffect::SetLocationEnd(const FVector& InEnd)
+{
+	if (!LocationEnd.Equals(InEnd))
+	{
+		LocationEnd = InEnd;
+		OnTransformOptionsChanged();
+	}
+}
+
+void UText3DLayoutTransformEffect::SetLocationEaseCurve(UCurveFloat* InEaseCurve)
+{
+	if (LocationEaseCurve == InEaseCurve)
+	{
+		return;
+	}
+
+	LocationEaseCurve = InEaseCurve;
+	OnTransformOptionsChanged();
+}
+
+void UText3DLayoutTransformEffect::SetScaleEnabled(bool bEnabled)
+{
+	if (bScaleEnabled != bEnabled)
+	{
+		bScaleEnabled = bEnabled;
+		OnTransformOptionsChanged();
+	}
+}
+
+void UText3DLayoutTransformEffect::SetScaleProgress(float Progress)
+{
+	Progress = FMath::Clamp(Progress, 0.0f, 100.0f);
+	if (!FMath::IsNearlyEqual(ScaleProgress, Progress))
+	{
+		ScaleProgress = Progress;
+		OnTransformOptionsChanged();
+	}
+}
+
+void UText3DLayoutTransformEffect::SetScaleOrder(EText3DCharacterEffectOrder Order)
+{
+	if (ScaleOrder != Order)
+	{
+		ScaleOrder = Order;
+		OnTransformOptionsChanged();
+	}
+}
+
+void UText3DLayoutTransformEffect::SetScaleBegin(const FVector& Value)
+{
+	if (!ScaleBegin.Equals(Value))
+	{
+		ScaleBegin = Value;
+		OnTransformOptionsChanged();
+	}
+}
+
+void UText3DLayoutTransformEffect::SetScaleEnd(const FVector& Value)
+{
+	if (!ScaleEnd.Equals(Value))
+	{
+		ScaleEnd = Value;
+		OnTransformOptionsChanged();
+	}
+}
+
+void UText3DLayoutTransformEffect::SetScaleEaseCurve(UCurveFloat* InEaseCurve)
+{
+	if (ScaleEaseCurve == InEaseCurve)
+	{
+		return;
+	}
+
+	ScaleEaseCurve = InEaseCurve;
+	OnTransformOptionsChanged();
+}
+
+void UText3DLayoutTransformEffect::SetRotationEnabled(bool bEnabled)
+{
+	if (bRotationEnabled != bEnabled)
+	{
+		bRotationEnabled = bEnabled;
+		OnTransformOptionsChanged();
+	}
+}
+
+void UText3DLayoutTransformEffect::SetRotationProgress(float Progress)
+{
+	Progress = FMath::Clamp(Progress, 0.0f, 100.0f);
+	if (!FMath::IsNearlyEqual(RotationProgress, Progress))
+	{
+		RotationProgress = Progress;
+		OnTransformOptionsChanged();
+	}
+}
+
+void UText3DLayoutTransformEffect::SetRotationOrder(EText3DCharacterEffectOrder Order)
+{
+	if (RotationOrder != Order)
+	{
+		RotationOrder = Order;
+		OnTransformOptionsChanged();
+	}
+}
+
+void UText3DLayoutTransformEffect::SetRotationBegin(const FRotator& Value)
+{
+	if (!RotationBegin.Equals(Value))
+	{
+		RotationBegin = Value;
+		OnTransformOptionsChanged();
+	}
+}
+
+void UText3DLayoutTransformEffect::SetRotationEnd(const FRotator& Value)
+{
+	if (!RotationEnd.Equals(Value))
+	{
+		RotationEnd = Value;
+		OnTransformOptionsChanged();
+	}
+}
+
+void UText3DLayoutTransformEffect::SetRotationEaseCurve(UCurveFloat* InEaseCurve)
+{
+	if (RotationEaseCurve == InEaseCurve)
+	{
+		return;
+	}
+
+	RotationEaseCurve = InEaseCurve;
+	OnTransformOptionsChanged();
+}
+
+UText3DLayoutTransformEffect::UText3DLayoutTransformEffect()
+{
+	OnTransformOptionsChanged();
+}
+
+#if WITH_EDITOR
+void UText3DLayoutTransformEffect::PostEditChangeProperty(FPropertyChangedEvent& InEvent)
+{
+	Super::PostEditChangeProperty(InEvent);
+
+	static const TSet<FName> PropertyNames =
+	{
+		GET_MEMBER_NAME_CHECKED(UText3DLayoutTransformEffect, bLocationEnabled),
+		GET_MEMBER_NAME_CHECKED(UText3DLayoutTransformEffect, LocationProgress),
+		GET_MEMBER_NAME_CHECKED(UText3DLayoutTransformEffect, LocationOrder),
+		GET_MEMBER_NAME_CHECKED(UText3DLayoutTransformEffect, LocationBegin),
+		GET_MEMBER_NAME_CHECKED(UText3DLayoutTransformEffect, LocationEnd),
+		GET_MEMBER_NAME_CHECKED(UText3DLayoutTransformEffect, LocationEaseCurve),
+		GET_MEMBER_NAME_CHECKED(UText3DLayoutTransformEffect, bRotationEnabled),
+		GET_MEMBER_NAME_CHECKED(UText3DLayoutTransformEffect, RotationProgress),
+		GET_MEMBER_NAME_CHECKED(UText3DLayoutTransformEffect, RotationOrder),
+		GET_MEMBER_NAME_CHECKED(UText3DLayoutTransformEffect, RotationBegin),
+		GET_MEMBER_NAME_CHECKED(UText3DLayoutTransformEffect, RotationEnd),
+		GET_MEMBER_NAME_CHECKED(UText3DLayoutTransformEffect, RotationEnd),
+		GET_MEMBER_NAME_CHECKED(UText3DLayoutTransformEffect, RotationEaseCurve),
+		GET_MEMBER_NAME_CHECKED(UText3DLayoutTransformEffect, bScaleEnabled),
+		GET_MEMBER_NAME_CHECKED(UText3DLayoutTransformEffect, ScaleProgress),
+		GET_MEMBER_NAME_CHECKED(UText3DLayoutTransformEffect, ScaleOrder),
+		GET_MEMBER_NAME_CHECKED(UText3DLayoutTransformEffect, ScaleBegin),
+		GET_MEMBER_NAME_CHECKED(UText3DLayoutTransformEffect, ScaleEnd),
+		GET_MEMBER_NAME_CHECKED(UText3DLayoutTransformEffect, ScaleEaseCurve),
+	};
+
+	if (PropertyNames.Contains(InEvent.GetMemberPropertyName()))
+	{
+		OnTransformOptionsChanged();
+	}
+}
+#endif
+
+void UText3DLayoutTransformEffect::OnTransformOptionsChanged()
+{
+	EText3DRendererFlags Flags = EText3DRendererFlags::Material;
+	EnumAddFlags(Flags, EText3DRendererFlags::Layout);
+	RequestUpdate(Flags);
+}
+
+void UText3DLayoutTransformEffect::ApplyEffect(uint32 InGlyphIndex, uint32 InGlyphCount)
+{
+	if (!bLocationEnabled && !bRotationEnabled && !bScaleEnabled)
+	{
+		return;
+	}
+
+	const UText3DComponent* Text3DComponent = GetText3DComponent();
+	UText3DCharacterBase* Character = Text3DComponent->GetCharacter(InGlyphIndex);
+
+	constexpr bool bReset = false;
+	FTransform& CharacterTransform = Character->GetTransform(bReset);
+
+	InGlyphIndex += 1;
+	InGlyphCount += 1;
+
+	if (bLocationEnabled)
+	{
+		const float Effect = CalculateEffect(InGlyphIndex, InGlyphCount, LocationOrder, LocationProgress, LocationEaseCurve);
+		const FVector Location = LocationBegin + (LocationEnd - LocationBegin) * Effect;
+		CharacterTransform.AddToTranslation(Location);
+	}
+
+	if (bScaleEnabled)
+	{
+		const float Effect = CalculateEffect(InGlyphIndex, InGlyphCount, ScaleOrder, ScaleProgress, ScaleEaseCurve);
+		const FVector Scale = ScaleBegin + (ScaleEnd - ScaleBegin) * Effect;
+		CharacterTransform.MultiplyScale3D(Scale);
+	}
+
+	if (bRotationEnabled)
+	{
+		const float Effect = CalculateEffect(InGlyphIndex, InGlyphCount, RotationOrder, RotationProgress, RotationEaseCurve);
+		const FRotator Rotator = RotationBegin + (RotationEnd - RotationBegin) * Effect;
+		CharacterTransform.ConcatenateRotation(Rotator.Quaternion());
+	}
+}
+
+int32 UText3DLayoutTransformEffect::GetEffectPosition(int32 Index, int32 Total, EText3DCharacterEffectOrder Order) const
+{
+	int32 Center = 0.5f * Total - 0.5f;
+	switch (Order)
+	{
+	case EText3DCharacterEffectOrder::FromCenter:
+		{
+			if (Index > Center)
+			{
+				Index = Total - Index - 1;
+			}
+
+			return Center - Index;
+		}
+
+	case EText3DCharacterEffectOrder::ToCenter:
+		{
+			if (Index > Center)
+			{
+				return Total - Index - 1;
+			}
+
+			return Index;
+		}
+
+	case EText3DCharacterEffectOrder::Opposite:
+		{
+			return Total - Index - 1;
+		}
+	default: ;
+	}
+
+	return Index;
+}
+
+float UText3DLayoutTransformEffect::CalculateEffect(int32 InIndex, int32 InTotal, EText3DCharacterEffectOrder InOrder, float InProgress, const UCurveFloat* InEaseCurve) const
+{
+	const int32 Position = GetEffectPosition(InIndex, InTotal, InOrder);
+	const float NormalizedProgress = FMath::Clamp(InProgress * 0.01f, 0.0f, 1.0f);
+	const float StaggerAmount = 1.0f / FMath::Max(1, InTotal);
+	const float CharacterStart = Position * StaggerAmount;
+	const float CharacterProgress = FMath::Clamp((NormalizedProgress - CharacterStart) / (1.0f - CharacterStart), 0.0f, 1.0f);
+	float CharacterEased = CharacterProgress;
+
+	if (InEaseCurve)
+	{
+		float StartTime = 0.0f;
+		float EndTime = 0.0f;
+		InEaseCurve->FloatCurve.GetTimeRange(StartTime, EndTime);
+
+		const float NormalizedTime = FMath::GetMappedRangeValueClamped(FVector2D(0.f, 1.f), FVector2D(StartTime, EndTime), CharacterProgress);
+		CharacterEased = InEaseCurve->FloatCurve.Eval(NormalizedTime);
+	}
+
+	return CharacterEased;
+}
diff --git a/Engine/Plugins/Experimental/Text3D/Source/Text3D/Private/Renderers/Text3DRendererBase.cpp b/Engine/Plugins/Experimental/Text3D/Source/Text3D/Private/Renderers/Text3DRendererBase.cpp
--- a/Engine/Plugins/Experimental/Text3D/Source/Text3D/Private/Renderers/Text3DRendererBase.cpp
+++ b/Engine/Plugins/Experimental/Text3D/Source/Text3D/Private/Renderers/Text3DRendererBase.cpp
@@ -4,10 +4,10 @@
 
 #include "Extensions/Text3DGeometryExtensionBase.h"
 #include "Extensions/Text3DCharacterExtensionBase.h"
+#include "Extensions/Text3DLayoutEffectBase.h"
 #include "Extensions/Text3DLayoutExtensionBase.h"
 #include "Extensions/Text3DMaterialExtensionBase.h"
 #include "Extensions/Text3DRenderingExtensionBase.h"
-#include "Extensions/Text3DTransformExtensionBase.h"
 #include "GameFramework/Actor.h"
 #include "Logs/Text3DLogs.h"
 #include "Text3DComponent.h"
@@ -62,7 +62,7 @@ void UText3DRendererBase::Update(EText3DRendererFlags InFlags)
 		Text3DComponent->GetRenderingExtension()
 	};
 
-	Extensions.Append(Text3DComponent->GetTransformExtensions());
+	Extensions.Append(Text3DComponent->GetLayoutEffects());
 
 	Extensions.RemoveAll([](const UText3DExtensionBase* InExtension)
 	{
diff --git a/Engine/Plugins/Experimental/Text3D/Source/Text3D/Private/Settings/Text3DProjectSettings.cpp b/Engine/Plugins/Experimental/Text3D/Source/Text3D/Private/Settings/Text3DProjectSettings.cpp
--- a/Engine/Plugins/Experimental/Text3D/Source/Text3D/Private/Settings/Text3DProjectSettings.cpp
+++ b/Engine/Plugins/Experimental/Text3D/Source/Text3D/Private/Settings/Text3DProjectSettings.cpp
@@ -29,8 +29,8 @@ FLazyName const UText3DProjectSettings::FMaterialParameters::TexturedVTiling(TEX
 
 FLazyName const UText3DProjectSettings::FMaterialParameters::BoundsOrigin(TEXT("TextBoundsOrigin"));
 FLazyName const UText3DProjectSettings::FMaterialParameters::BoundsSize(TEXT("TextBoundsSize"));
-FLazyName const UText3DProjectSettings::FMaterialParameters::TextPosition(TEXT("TextActorPosition"));
 
+FLazyName const UText3DProjectSettings::FMaterialParameters::MaskEnabled(TEXT("MaskEnabled"));
 FLazyName const UText3DProjectSettings::FMaterialParameters::MaskRotation(TEXT("MaskRotationDegrees"));
 FLazyName const UText3DProjectSettings::FMaterialParameters::MaskOffset(TEXT("MaskOffset"));
 FLazyName const UText3DProjectSettings::FMaterialParameters::MaskSmoothness(TEXT("MaskSmoothness"));
diff --git a/Engine/Plugins/Experimental/Text3D/Source/Text3D/Private/Text3DComponent.cpp b/Engine/Plugins/Experimental/Text3D/Source/Text3D/Private/Text3DComponent.cpp
--- a/Engine/Plugins/Experimental/Text3D/Source/Text3D/Private/Text3DComponent.cpp
+++ b/Engine/Plugins/Experimental/Text3D/Source/Text3D/Private/Text3DComponent.cpp
@@ -11,6 +11,7 @@
 #include "Extensions/Text3DDefaultLayoutExtension.h"
 #include "Extensions/Text3DDefaultMaterialExtension.h"
 #include "Extensions/Text3DDefaultRenderingExtension.h"
+#include "Extensions/Text3DLayoutEffectBase.h"
 #include "Fonts/CompositeFont.h"
 #include "Internationalization/Regex.h"
 #include "Logs/Text3DLogs.h"
@@ -361,7 +362,7 @@ void UText3DComponent::PostEditChangeProperty(FPropertyChangedEvent& PropertyCha
 	{
 		OnFontPropertiesChanged();
 	}
-	else if (MemberPropertyName == GET_MEMBER_NAME_CHECKED(UText3DComponent, TransformExtensions))
+	else if (MemberPropertyName == GET_MEMBER_NAME_CHECKED(UText3DComponent, LayoutEffects))
 	{
 		EText3DRendererFlags Flags = EText3DRendererFlags::Layout;
 		EnumAddFlags(Flags, EText3DRendererFlags::Material);
@@ -564,6 +565,67 @@ UText3DRendererBase* UText3DComponent::GetTextRenderer() const
 	return TextRenderer;
 }
 
+UText3DLayoutExtensionBase* UText3DComponent::GetLayoutExtension(TSubclassOf<UText3DLayoutExtensionBase> InExtensionClass)
+{
+	if (!LayoutExtension || !InExtensionClass.Get())
+	{
+		return nullptr;
+	}
+	
+	return LayoutExtension->IsA(InExtensionClass.Get()) ? LayoutExtension : nullptr;
+}
+
+UText3DMaterialExtensionBase* UText3DComponent::GetMaterialExtension(TSubclassOf<UText3DMaterialExtensionBase> InExtensionClass)
+{
+	if (!MaterialExtension || !InExtensionClass.Get())
+	{
+		return nullptr;
+	}
+	
+	return MaterialExtension->IsA(InExtensionClass.Get()) ? MaterialExtension : nullptr;
+}
+
+UText3DGeometryExtensionBase* UText3DComponent::GetGeometryExtension(TSubclassOf<UText3DGeometryExtensionBase> InExtensionClass)
+{
+	if (!GeometryExtension || !InExtensionClass.Get())
+	{
+		return nullptr;
+	}
+	
+	return GeometryExtension->IsA(InExtensionClass.Get()) ? GeometryExtension : nullptr;
+}
+
+UText3DRenderingExtensionBase* UText3DComponent::GetRenderingExtension(TSubclassOf<UText3DRenderingExtensionBase> InExtensionClass)
+{
+	if (!RenderingExtension || !InExtensionClass.Get())
+	{
+		return nullptr;
+	}
+	
+	return RenderingExtension->IsA(InExtensionClass.Get()) ? RenderingExtension : nullptr;
+}
+
+TArray<UText3DLayoutEffectBase*> UText3DComponent::GetLayoutEffects(TSubclassOf<UText3DLayoutEffectBase> InEffectClass)
+{
+	TArray<UText3DLayoutEffectBase*> Effects;
+
+	if (!InEffectClass.Get())
+	{
+		return Effects;
+	}
+
+	Effects.Reserve(LayoutEffects.Num());
+	for (UText3DLayoutEffectBase* LayoutEffect : LayoutEffects)
+	{
+		if (LayoutEffect && LayoutEffect->IsA(InEffectClass))
+		{
+			Effects.Add(LayoutEffect);
+		}
+	}
+
+	return Effects;
+}
+
 uint16 UText3DComponent::GetCharacterCount() const
 {
 	return CharacterExtension->GetCharacterCount();
diff --git a/Engine/Plugins/Experimental/Text3D/Source/Text3D/Public/Extensions/Text3DDefaultMaterialExtension.h b/Engine/Plugins/Experimental/Text3D/Source/Text3D/Public/Extensions/Text3DDefaultMaterialExtension.h
--- a/Engine/Plugins/Experimental/Text3D/Source/Text3D/Public/Extensions/Text3DDefaultMaterialExtension.h
+++ b/Engine/Plugins/Experimental/Text3D/Source/Text3D/Public/Extensions/Text3DDefaultMaterialExtension.h
@@ -150,6 +150,15 @@ public:
 		return Opacity;
 	}
 
+	UFUNCTION(BlueprintCallable, Category = "Text3D|Material")
+	TEXT3D_API void SetUseMask(bool bInUseMask);
+
+	UFUNCTION(BlueprintPure, Category = "Text3D|Material")
+	bool GetUseMask() const
+	{
+		return bUseMask;
+	}
+
 	UFUNCTION(BlueprintCallable, Category = "Text3D|Material")
 	TEXT3D_API void SetMaskOffset(float InMaskOffset);
 
@@ -177,6 +186,15 @@ public:
 		return MaskRotation;
 	}
 
+	UFUNCTION(BlueprintCallable, Category = "Text3D|Material")
+	TEXT3D_API void SetUseSingleMaterial(bool bInUse);
+
+	UFUNCTION(BlueprintPure, Category = "Text3D|Material")
+	bool GetUseSingleMaterial() const
+	{
+		return bUseSingleMaterial;
+	}
+
 	UFUNCTION(BlueprintCallable, Category = "Text3D|Material")
 	TEXT3D_API void SetFrontMaterial(UMaterialInterface* InMaterial);
 
@@ -241,6 +259,7 @@ protected:
 	void SetScalarParameter(TConstArrayView<UMaterialInstanceDynamic*> InMaterials, FName InKey, float InValue);
 	void SetTextureParameter(TConstArrayView<UMaterialInstanceDynamic*> InMaterials, FName InKey, UTexture* InValue);
 	void OnMaterialOptionsChanged();
+	void OnCustomMaterialChanged();
 
 	UPROPERTY(EditAnywhere, Getter, Setter, Category="Material", meta=(AllowPrivateAccess="true"))
 	EText3DMaterialStyle Style = EText3DMaterialStyle::Solid;
@@ -287,29 +306,37 @@ protected:
 	UPROPERTY(EditAnywhere, Getter, Setter, Category="Material", meta=(ClampMin=0.0, ClampMax=1.0, EditCondition="Style != EText3DMaterialStyle::Custom && BlendMode == EText3DMaterialBlendMode::Translucent", EditConditionHides, AllowPrivateAccess="true"))
 	float Opacity = 1.f;
 
-	UPROPERTY(EditAnywhere, Getter, Setter, Category="Material", meta=(ClampMin=0.0, ClampMax=1.0, EditCondition="Style != EText3DMaterialStyle::Custom && BlendMode == EText3DMaterialBlendMode::Translucent", EditConditionHides, AllowPrivateAccess="true"))
+	/** Enable text shader mask */
+	UPROPERTY(EditAnywhere, Getter="GetUseMask", Setter="SetUseMask", Category="Material", meta=(EditCondition="Style != EText3DMaterialStyle::Custom && BlendMode == EText3DMaterialBlendMode::Translucent", EditConditionHides, AllowPrivateAccess="true"))
+	bool bUseMask = false;
+
+	UPROPERTY(EditAnywhere, Getter, Setter, Category="Material", meta=(ClampMin=0.0, ClampMax=1.0, EditCondition="bUseMask && Style != EText3DMaterialStyle::Custom && BlendMode == EText3DMaterialBlendMode::Translucent", EditConditionHides, AllowPrivateAccess="true"))
 	float MaskOffset = 1.f;
 
-	UPROPERTY(EditAnywhere, Getter, Setter, Category="Material", meta=(ClampMin=0.0, ClampMax=1.0, EditCondition="Style != EText3DMaterialStyle::Custom && BlendMode == EText3DMaterialBlendMode::Translucent", EditConditionHides, AllowPrivateAccess="true"))
+	UPROPERTY(EditAnywhere, Getter, Setter, Category="Material", meta=(ClampMin=0.0, ClampMax=1.0, EditCondition="bUseMask && Style != EText3DMaterialStyle::Custom && BlendMode == EText3DMaterialBlendMode::Translucent", EditConditionHides, AllowPrivateAccess="true"))
 	float MaskSmoothness = 0.1f;
 
-	UPROPERTY(EditAnywhere, Getter, Setter, Category="Material", meta=(ClampMin=0.0, ClampMax=1.0, EditCondition="Style != EText3DMaterialStyle::Custom && BlendMode == EText3DMaterialBlendMode::Translucent", EditConditionHides, AllowPrivateAccess="true"))
+	UPROPERTY(EditAnywhere, Getter, Setter, Category="Material", meta=(ClampMin=0.0, ClampMax=1.0, EditCondition="bUseMask && Style != EText3DMaterialStyle::Custom && BlendMode == EText3DMaterialBlendMode::Translucent", EditConditionHides, AllowPrivateAccess="true"))
 	float MaskRotation = 0.0f;
 
+	/** Use primary material for all available slots */
+	UPROPERTY(EditAnywhere, Getter="GetUseSingleMaterial", Setter="SetUseSingleMaterial", Category = "Material", meta = (EditCondition="Style == EText3DMaterialStyle::Custom", EditConditionHides, AllowPrivateAccess = "true"))
+	bool bUseSingleMaterial = false;
+
 	/** Material for the front part */
 	UPROPERTY(EditAnywhere, Getter, Setter, Category = "Material", meta = (EditCondition="Style == EText3DMaterialStyle::Custom", EditConditionHides, AllowPrivateAccess = "true"))
 	TObjectPtr<UMaterialInterface> FrontMaterial;
 
 	/** Material for the bevel part */
-	UPROPERTY(EditAnywhere, Getter, Setter, Category = "Material", meta = (EditCondition="Style == EText3DMaterialStyle::Custom", EditConditionHides, AllowPrivateAccess = "true"))
+	UPROPERTY(EditAnywhere, Getter, Setter, Category = "Material", meta = (EditCondition="!bUseSingleMaterial && Style == EText3DMaterialStyle::Custom", EditConditionHides, AllowPrivateAccess = "true"))
 	TObjectPtr<UMaterialInterface> BevelMaterial;
 
 	/** Material for the extruded part */
-	UPROPERTY(EditAnywhere, Getter, Setter, Category = "Material", meta = (EditCondition="Style == EText3DMaterialStyle::Custom", EditConditionHides, AllowPrivateAccess = "true"))
+	UPROPERTY(EditAnywhere, Getter, Setter, Category = "Material", meta = (EditCondition="!bUseSingleMaterial && Style == EText3DMaterialStyle::Custom", EditConditionHides, AllowPrivateAccess = "true"))
 	TObjectPtr<UMaterialInterface> ExtrudeMaterial;
 
 	/** Material for the back part */
-	UPROPERTY(EditAnywhere, Getter, Setter, Category = "Material", meta = (EditCondition="Style == EText3DMaterialStyle::Custom", EditConditionHides, AllowPrivateAccess = "true"))
+	UPROPERTY(EditAnywhere, Getter, Setter, Category = "Material", meta = (EditCondition="!bUseSingleMaterial && Style == EText3DMaterialStyle::Custom", EditConditionHides, AllowPrivateAccess = "true"))
 	TObjectPtr<UMaterialInterface> BackMaterial;
 
 	/** Cached group dynamic materials created during session based on their options */
diff --git a/Engine/Plugins/Experimental/Text3D/Source/Text3D/Public/Extensions/Text3DDefaultTransformExtension.h b/Engine/Plugins/Experimental/Text3D/Source/Text3D/Public/Extensions/Text3DDefaultTransformExtension.h
deleted file mode 100644
diff --git a/Engine/Plugins/Experimental/Text3D/Source/Text3D/Public/Extensions/Text3DEffectExtensionBase.h b/Engine/Plugins/Experimental/Text3D/Source/Text3D/Public/Extensions/Text3DEffectExtensionBase.h
new file mode 100644
--- /dev/null
+++ b/Engine/Plugins/Experimental/Text3D/Source/Text3D/Public/Extensions/Text3DEffectExtensionBase.h
@@ -0,0 +1,76 @@
+// Copyright Epic Games, Inc. All Rights Reserved.
+
+#pragma once
+
+#include "Text3DExtensionBase.h"
+#include "Text3DTypes.h"
+#include "Text3DEffectExtensionBase.generated.h"
+
+struct FText3DRange
+{
+	static constexpr uint16 Unbound = std::numeric_limits<uint16>::max();
+
+	FText3DRange() = default;
+	explicit FText3DRange(uint16 InStartIndex, uint16 InCount)
+		: StartIndex(InStartIndex)
+		, Count(InCount)
+	{
+	}
+
+	bool IsUnbound() const
+	{
+		return IsStartUnbound() && IsEndUnbound();
+	}
+
+	bool IsEndUnbound() const
+	{
+		return Count == Unbound;
+	}
+
+	bool IsStartUnbound() const
+	{
+		return StartIndex == Unbound;
+	}
+
+	bool IsInRange(uint16 InIndex) const
+	{
+		return IsUnbound()
+			|| (IsStartUnbound() && InIndex < Count)
+			|| (IsEndUnbound() && InIndex >= StartIndex)
+			|| (InIndex >= StartIndex && InIndex <= (StartIndex + Count));
+	}
+
+	uint16 StartIndex = Unbound;
+	uint16 Count = Unbound;
+};
+
+/** Extension for custom effects on Text3D */
+UCLASS(MinimalAPI, Abstract)
+class UText3DEffectExtensionBase : public UText3DExtensionBase
+{
+	GENERATED_BODY()
+
+public:
+	UText3DEffectExtensionBase()
+		: UText3DExtensionBase(/** Priority */0)
+	{}
+
+	explicit UText3DEffectExtensionBase(uint16 InPriority)
+		: UText3DExtensionBase(InPriority)
+	{}
+
+protected:
+	//~ Begin UText3DExtensionBase
+	virtual EText3DExtensionResult PreRendererUpdate(EText3DRendererFlags InFlag) override;
+	virtual EText3DExtensionResult PostRendererUpdate(EText3DRendererFlags InFlag) override;
+	//~ End UText3DExtensionBase
+
+	virtual FText3DRange GetTargetRange() const
+	{
+		return FText3DRange();
+	}
+
+	virtual void ApplyEffect(uint32 InCharacterIndex, uint32 InCharacterCount)
+	{
+	}
+};
\ No newline at end of file
diff --git a/Engine/Plugins/Experimental/Text3D/Source/Text3D/Public/Extensions/Text3DLayoutEffectBase.h b/Engine/Plugins/Experimental/Text3D/Source/Text3D/Public/Extensions/Text3DLayoutEffectBase.h
new file mode 100644
--- /dev/null
+++ b/Engine/Plugins/Experimental/Text3D/Source/Text3D/Public/Extensions/Text3DLayoutEffectBase.h
@@ -0,0 +1,18 @@
+// Copyright Epic Games, Inc. All Rights Reserved.
+
+#pragma once
+
+#include "Text3DEffectExtensionBase.h"
+#include "Text3DLayoutEffectBase.generated.h"
+
+/** Extension for layout effects on Text3D */
+UCLASS(MinimalAPI, Abstract)
+class UText3DLayoutEffectBase : public UText3DEffectExtensionBase
+{
+	GENERATED_BODY()
+
+public:
+	UText3DLayoutEffectBase()
+		: UText3DEffectExtensionBase(/** Priority */3)
+	{}
+};
diff --git a/Engine/Plugins/Experimental/Text3D/Source/Text3D/Public/Extensions/Text3DLayoutTransformEffect.h b/Engine/Plugins/Experimental/Text3D/Source/Text3D/Public/Extensions/Text3DLayoutTransformEffect.h
new file mode 100644
--- /dev/null
+++ b/Engine/Plugins/Experimental/Text3D/Source/Text3D/Public/Extensions/Text3DLayoutTransformEffect.h
@@ -0,0 +1,243 @@
+// Copyright Epic Games, Inc. All Rights Reserved.
+
+#pragma once
+
+#include "Text3DLayoutEffectBase.h"
+#include "Text3DTypes.h"
+#include "Text3DLayoutTransformEffect.generated.h"
+
+class UCurveFloat;
+
+/** Extension that handles transform data for Text3D */
+UCLASS(MinimalAPI, EditInlineNew, ClassGroup=Text3D, DisplayName="TransformEffect")
+class UText3DLayoutTransformEffect : public UText3DLayoutEffectBase
+{
+	GENERATED_BODY()
+
+public:
+	UFUNCTION(BlueprintCallable, Category = "Text3D|Transform")
+	TEXT3D_API void SetLocationEnabled(bool bEnabled);
+
+	bool GetLocationEnabled() const
+	{
+		return bLocationEnabled;
+	}
+
+	UFUNCTION(BlueprintCallable, Category = "Text3D|Transform")
+	TEXT3D_API void SetLocationProgress(float Progress);
+
+	float GetLocationProgress() const
+	{
+		return LocationProgress;
+	}
+
+	UFUNCTION(BlueprintCallable, Category = "Text3D|Transform")
+	TEXT3D_API void SetLocationOrder(EText3DCharacterEffectOrder Order);
+
+	EText3DCharacterEffectOrder GetLocationOrder() const
+	{
+		return LocationOrder;
+	}
+
+	UFUNCTION(BlueprintCallable, Category = "Text3D|Transform")
+	TEXT3D_API void SetLocationBegin(const FVector& InBegin);
+
+	FVector GetLocationBegin() const
+	{
+		return LocationBegin;
+	}
+
+	UFUNCTION(BlueprintCallable, Category = "Text3D|Transform")
+	TEXT3D_API void SetLocationEnd(const FVector& InEnd);
+
+	FVector GetLocationEnd() const
+	{
+		return LocationEnd;
+	}
+
+	UFUNCTION(BlueprintCallable, Category = "Text3D|Transform")
+	TEXT3D_API void SetLocationEaseCurve(UCurveFloat* InEaseCurve);
+
+	UCurveFloat* GetLocationEaseCurve() const
+	{
+		return LocationEaseCurve;
+	}
+
+	UFUNCTION(BlueprintCallable, Category = "Text3D|Transform")
+	TEXT3D_API void SetRotationEnabled(bool bEnabled);
+
+	bool GetRotationEnabled() const
+	{
+		return bRotationEnabled;
+	}
+
+	UFUNCTION(BlueprintCallable, Category = "Text3D|Transform")
+	TEXT3D_API void SetRotationProgress(float Progress);
+
+	float GetRotationProgress() const
+	{
+		return RotationProgress;
+	}
+
+	UFUNCTION(BlueprintCallable, Category = "Text3D|Transform")
+	TEXT3D_API void SetRotationOrder(EText3DCharacterEffectOrder Order);
+
+	EText3DCharacterEffectOrder GetRotationOrder() const
+	{
+		return RotationOrder;
+	}
+
+	UFUNCTION(BlueprintCallable, Category = "Text3D|Transform")
+	TEXT3D_API void SetRotationBegin(const FRotator& Value);
+
+	FRotator GetRotationBegin() const
+	{
+		return RotationBegin;
+	}
+
+	UFUNCTION(BlueprintCallable, Category = "Text3D|Transform")
+	TEXT3D_API void SetRotationEnd(const FRotator& Value);
+
+	FRotator GetRotationEnd() const
+	{
+		return RotationEnd;
+	}
+
+	UFUNCTION(BlueprintCallable, Category = "Text3D|Transform")
+	TEXT3D_API void SetRotationEaseCurve(UCurveFloat* InEaseCurve);
+
+	UCurveFloat* GetRotationEaseCurve() const
+	{
+		return RotationEaseCurve;
+	}
+
+	UFUNCTION(BlueprintCallable, Category = "Text3D|Transform")
+	TEXT3D_API void SetScaleEnabled(bool bEnabled);
+
+	bool GetScaleEnabled() const
+	{
+		return bScaleEnabled;
+	}
+
+	UFUNCTION(BlueprintCallable, Category = "Text3D|Transform")
+	TEXT3D_API void SetScaleProgress(float Progress);
+
+	float GetScaleProgress() const
+	{
+		return ScaleProgress;
+	}
+
+	UFUNCTION(BlueprintCallable, Category = "Text3D|Transform")
+	TEXT3D_API void SetScaleOrder(EText3DCharacterEffectOrder Order);
+
+	EText3DCharacterEffectOrder GetScaleOrder() const
+	{
+		return ScaleOrder;
+	}
+
+	UFUNCTION(BlueprintCallable, Category = "Text3D|Transform")
+	TEXT3D_API void SetScaleBegin(const FVector& Value);
+
+	FVector GetScaleBegin() const
+	{
+		return ScaleBegin;
+	}
+
+	UFUNCTION(BlueprintCallable, Category = "Text3D|Transform")
+	TEXT3D_API void SetScaleEnd(const FVector& Value);
+
+	FVector GetScaleEnd() const
+	{
+		return ScaleEnd;
+	}
+
+	UFUNCTION(BlueprintCallable, Category = "Text3D|Transform")
+	TEXT3D_API void SetScaleEaseCurve(UCurveFloat* InEaseCurve);
+
+	UCurveFloat* GetScaleEaseCurve() const
+	{
+		return ScaleEaseCurve;
+	}
+
+
+protected:
+	UText3DLayoutTransformEffect();
+
+	//~ Begin UObject
+#if WITH_EDITOR
+	virtual void PostEditChangeProperty(FPropertyChangedEvent& InEvent) override;
+#endif
+	//~ End UObject
+
+	//~ Begin UText3DTransformExtensionBase
+	virtual void ApplyEffect(uint32 InGlyphIndex, uint32 InGlyphCount) override;
+	//~ End UText3DTransformExtensionBase
+
+	void OnTransformOptionsChanged();
+
+	int32 GetEffectPosition(int32 Index, int32 Total, EText3DCharacterEffectOrder Order) const;
+	float CalculateEffect(int32 InIndex, int32 InTotal, EText3DCharacterEffectOrder InOrder, float InProgress, const UCurveFloat* InEaseCurve) const;
+
+	// Location
+
+	UPROPERTY(EditAnywhere, Getter="GetLocationEnabled", Setter="SetLocationEnabled", Category = "Transform")
+	bool bLocationEnabled = false;
+
+	UPROPERTY(EditAnywhere, Getter, Setter, Category = "Transform", meta = (EditCondition = "bLocationEnabled", EditConditionHides, ClampMin = 0, ClampMax = 100, Units=Percent))
+	float LocationProgress = 0.f;
+
+	UPROPERTY(EditAnywhere, Getter, Setter, Category = "Transform", meta = (EditCondition = "bLocationEnabled", EditConditionHides))
+	EText3DCharacterEffectOrder LocationOrder = EText3DCharacterEffectOrder::Normal;
+
+	UPROPERTY(EditAnywhere, Getter, Setter, Category = "Transform", meta = (EditCondition = "bLocationEnabled", EditConditionHides))
+	FVector LocationBegin = FVector::ZeroVector;
+
+	UPROPERTY(EditAnywhere, Getter, Setter, Category = "Transform", meta = (EditCondition = "bLocationEnabled", EditConditionHides))
+	FVector LocationEnd = FVector(100.f, 0.f, 0.f);
+
+	/** Provide a 0-1 ease curve, leaving this unset will result in linear ease */
+	UPROPERTY(EditAnywhere, Setter, Getter, Category = "Transform", meta = (EditCondition = "bLocationEnabled", EditConditionHides))
+	TObjectPtr<UCurveFloat> LocationEaseCurve;
+
+	// Rotate
+
+	UPROPERTY(EditAnywhere, Getter="GetRotationEnabled", Setter="SetRotationEnabled", Category = "Transform")
+	bool bRotationEnabled = false;
+
+	UPROPERTY(EditAnywhere, Getter, Setter, Category = "Transform", meta = (EditCondition = "bRotationEnabled", EditConditionHides, ClampMin = 0, ClampMax = 100, Units=Percent))
+	float RotationProgress = 0.0f;
+
+	UPROPERTY(EditAnywhere, Getter, Setter, Category = "Transform", meta = (EditCondition = "bRotationEnabled", EditConditionHides))
+	EText3DCharacterEffectOrder RotationOrder = EText3DCharacterEffectOrder::Normal;
+
+	UPROPERTY(EditAnywhere, Getter, Setter, Category = "Transform", meta = (EditCondition = "bRotationEnabled", EditConditionHides))
+	FRotator RotationBegin = FRotator(-90.f, 0.f, 0.f);
+
+	UPROPERTY(EditAnywhere, Getter, Setter, Category = "Transform", meta = (EditCondition = "bRotationEnabled", EditConditionHides))
+	FRotator RotationEnd = FRotator(0.f, 0.f, 0.f);
+
+	/** Provide a 0-1 ease curve, leaving this unset will result in linear ease */
+	UPROPERTY(EditAnywhere, Setter, Getter, Category = "Transform", meta = (EditCondition = "bRotationEnabled", EditConditionHides))
+	TObjectPtr<UCurveFloat> RotationEaseCurve;
+
+	// Scale
+
+	UPROPERTY(EditAnywhere, Getter="GetScaleEnabled", Setter="SetScaleEnabled", Category = "Transform")
+	bool bScaleEnabled = false;
+
+	UPROPERTY(EditAnywhere, Getter, Setter, Category = "Transform", meta = (EditCondition = "bScaleEnabled", EditConditionHides, ClampMin = 0, ClampMax = 100, Units=Percent))
+	float ScaleProgress = 0.f;
+
+	UPROPERTY(EditAnywhere, Getter, Setter, Category = "Transform", meta = (EditCondition = "bScaleEnabled", EditConditionHides))
+	EText3DCharacterEffectOrder ScaleOrder = EText3DCharacterEffectOrder::Normal;
+
+	UPROPERTY(EditAnywhere, Getter, Setter, Category = "Transform", meta = (EditCondition = "bScaleEnabled", EditConditionHides, ClampMin = 0))
+	FVector ScaleBegin = FVector(1.f, 0.f, 0.f);
+
+	UPROPERTY(EditAnywhere, Getter, Setter, Category = "Transform", meta = (EditCondition = "bScaleEnabled", EditConditionHides, ClampMin = 0))
+	FVector ScaleEnd = FVector::OneVector;
+
+	/** Provide a 0-1 ease curve, leaving this unset will result in linear ease */
+	UPROPERTY(EditAnywhere, Setter, Getter, Category = "Transform", meta = (EditCondition = "bScaleEnabled", EditConditionHides))
+	TObjectPtr<UCurveFloat> ScaleEaseCurve;
+};
diff --git a/Engine/Plugins/Experimental/Text3D/Source/Text3D/Public/Extensions/Text3DTransformExtensionBase.h b/Engine/Plugins/Experimental/Text3D/Source/Text3D/Public/Extensions/Text3DTransformExtensionBase.h
deleted file mode 100644
diff --git a/Engine/Plugins/Experimental/Text3D/Source/Text3D/Public/Settings/Text3DProjectSettings.h b/Engine/Plugins/Experimental/Text3D/Source/Text3D/Public/Settings/Text3DProjectSettings.h
--- a/Engine/Plugins/Experimental/Text3D/Source/Text3D/Public/Settings/Text3DProjectSettings.h
+++ b/Engine/Plugins/Experimental/Text3D/Source/Text3D/Public/Settings/Text3DProjectSettings.h
@@ -50,6 +50,7 @@ public:
 		static const FLazyName TextPosition;
 
 		// Mask
+		static const FLazyName MaskEnabled;
 		static const FLazyName MaskRotation;
 		static const FLazyName MaskOffset;
 		static const FLazyName MaskSmoothness;
diff --git a/Engine/Plugins/Experimental/Text3D/Source/Text3D/Public/Text3DComponent.h b/Engine/Plugins/Experimental/Text3D/Source/Text3D/Public/Text3DComponent.h
--- a/Engine/Plugins/Experimental/Text3D/Source/Text3D/Public/Text3DComponent.h
+++ b/Engine/Plugins/Experimental/Text3D/Source/Text3D/Public/Text3DComponent.h
@@ -19,7 +19,7 @@ class UText3DGeometryExtensionBase;
 class UText3DMaterialExtensionBase;
 class UText3DRendererBase;
 class UText3DRenderingExtensionBase;
-class UText3DTransformExtensionBase;
+class UText3DLayoutEffectBase;
 struct FText3DShapedGlyphText;
 struct FTypefaceEntry;
 
@@ -247,8 +247,10 @@ public:
 	/** Set the typeface */
 	TEXT3D_API void SetTypeface(const FName InTypeface);
 
+	UFUNCTION(BlueprintCallable, Category = "Text3D")
 	void SetTextRendererClass(const TSubclassOf<UText3DRendererBase>& InClass);
-	
+
+	UFUNCTION(BlueprintPure, Category = "Text3D")
 	TSubclassOf<UText3DRendererBase> GetTextRendererClass() const
 	{
 		return TextRendererClass;
@@ -269,6 +271,9 @@ public:
 	
 	UText3DRendererBase* GetTextRenderer() const;
 
+	UFUNCTION(BlueprintCallable, Category="Text3D", meta=(DeterminesOutputType="InExtensionClass"))
+	UText3DLayoutExtensionBase* GetLayoutExtension(TSubclassOf<UText3DLayoutExtensionBase> InExtensionClass);
+
 	UText3DLayoutExtensionBase* GetLayoutExtension() const
 	{
 		return LayoutExtension;
@@ -280,6 +285,9 @@ public:
 		return Cast<InType>(LayoutExtension);
 	}
 
+	UFUNCTION(BlueprintCallable, Category="Text3D", meta=(DeterminesOutputType="InExtensionClass"))
+	UText3DMaterialExtensionBase* GetMaterialExtension(TSubclassOf<UText3DMaterialExtensionBase> InExtensionClass);
+
 	UText3DMaterialExtensionBase* GetMaterialExtension() const
 	{
 		return MaterialExtension;
@@ -291,6 +299,9 @@ public:
 		return Cast<InType>(MaterialExtension);
 	}
 
+	UFUNCTION(BlueprintCallable, Category="Text3D", meta=(DeterminesOutputType="InExtensionClass"))
+	UText3DGeometryExtensionBase* GetGeometryExtension(TSubclassOf<UText3DGeometryExtensionBase> InExtensionClass);
+
 	UText3DGeometryExtensionBase* GetGeometryExtension() const
 	{
 		return GeometryExtension;
@@ -301,6 +312,9 @@ public:
 	{
 		return Cast<InType>(GeometryExtension);
 	}
+
+	UFUNCTION(BlueprintCallable, Category="Text3D", meta=(DeterminesOutputType="InExtensionClass"))
+	UText3DRenderingExtensionBase* GetRenderingExtension(TSubclassOf<UText3DRenderingExtensionBase> InExtensionClass);
 	
 	UText3DRenderingExtensionBase* GetRenderingExtension() const
 	{
@@ -313,9 +327,12 @@ public:
 		return Cast<InType>(RenderingExtension);
 	}
 
-	TConstArrayView<UText3DTransformExtensionBase*> GetTransformExtensions() const
+	UFUNCTION(BlueprintCallable, Category="Text3D", meta=(DeterminesOutputType="InEffectClass"))
+	TArray<UText3DLayoutEffectBase*> GetLayoutEffects(TSubclassOf<UText3DLayoutEffectBase> InEffectClass);
+
+	TConstArrayView<UText3DLayoutEffectBase*> GetLayoutEffects() const
 	{
-		return TransformExtensions;
+		return LayoutEffects;
 	}
 	
 	UText3DCharacterExtensionBase* GetCharacterExtension() const
@@ -439,6 +456,9 @@ protected:
 	UPROPERTY(VisibleInstanceOnly, Instanced, Category="Layout", meta=(AllowPrivateAccess="true"))
 	TObjectPtr<UText3DLayoutExtensionBase> LayoutExtension;
 
+	UPROPERTY(EditInstanceOnly, Instanced, Category="Layout|Effects", meta=(AllowPrivateAccess="true"))
+	TArray<TObjectPtr<UText3DLayoutEffectBase>> LayoutEffects;
+
 	UPROPERTY(VisibleInstanceOnly, Instanced, Category="Geometry", meta=(AllowPrivateAccess="true"))
 	TObjectPtr<UText3DGeometryExtensionBase> GeometryExtension;
 
@@ -448,9 +468,6 @@ protected:
 	UPROPERTY(VisibleInstanceOnly, Instanced, Category="Rendering", meta=(AllowPrivateAccess="true"))
 	TObjectPtr<UText3DRenderingExtensionBase> RenderingExtension;
 
-	UPROPERTY(EditInstanceOnly, Instanced, Category="Transform", meta=(AllowPrivateAccess="true"))
-	TArray<TObjectPtr<UText3DTransformExtensionBase>> TransformExtensions;
-
 private:
 	/** Active renderer in charge of generating text geometry */
 	UPROPERTY()
diff --git a/Engine/Plugins/Experimental/Text3D/Source/Text3DEditor/Private/Customizations/Text3DEditorTextComponentDetailCustomization.cpp b/Engine/Plugins/Experimental/Text3D/Source/Text3DEditor/Private/Customizations/Text3DEditorTextComponentDetailCustomization.cpp
--- a/Engine/Plugins/Experimental/Text3D/Source/Text3DEditor/Private/Customizations/Text3DEditorTextComponentDetailCustomization.cpp
+++ b/Engine/Plugins/Experimental/Text3D/Source/Text3DEditor/Private/Customizations/Text3DEditorTextComponentDetailCustomization.cpp
@@ -26,7 +26,6 @@ namespace UE::Text3DEditor::Customization
 			const TSharedRef<FPropertySection> LayoutSection = PropertyModule.FindOrCreateSection(ComponentClassName, LayoutSectionName, FText::FromName(LayoutSectionName));
 			LayoutSection->AddCategory(TEXT("Layout"));
 			LayoutSection->AddCategory(TEXT("Character"));
-			LayoutSection->AddCategory(TEXT("Transform"));
 
 			const FName RenderingSectionName = TEXT("Rendering");
 			const TSharedRef<FPropertySection> RenderingSection = PropertyModule.FindOrCreateSection(ComponentClassName, RenderingSectionName, FText::FromName(RenderingSectionName));
