MotionDesign : Text3D
- Added effects section for text (layout effects) in details panel

diff --git a/Engine/Plugins/Experimental/Avalanche/Source/AvalancheTextEditor/Private/AvaTextEditorModule.cpp b/Engine/Plugins/Experimental/Avalanche/Source/AvalancheTextEditor/Private/AvaTextEditorModule.cpp
--- a/Engine/Plugins/Experimental/Avalanche/Source/AvalancheTextEditor/Private/AvaTextEditorModule.cpp
+++ b/Engine/Plugins/Experimental/Avalanche/Source/AvalancheTextEditor/Private/AvaTextEditorModule.cpp
@@ -66,7 +66,7 @@ void FAvaTextEditorModule::RegisterDynamicMaterialPropertyGenerator()
 			FDMGetObjectMaterialPropertiesDelegate::CreateLambda(
 				[](UObject* InObject)
 				{
-					UClass* MaterialExtensionClass = UText3DDefaultMaterialExtension::StaticClass();
+					const UClass* MaterialExtensionClass = UText3DDefaultMaterialExtension::StaticClass();
 					static TArray<FProperty*> Text3DProperties =
 					{
 						MaterialExtensionClass->FindPropertyByName("FrontMaterial"),
@@ -91,6 +91,11 @@ void FAvaTextEditorModule::RegisterDynamicMaterialPropertyGenerator()
 									{
 										Properties.Emplace(MaterialExtension, Property);
 									}
+
+									if (MaterialExtension->GetUseSingleMaterial())
+									{
+										break;
+									}
 								}
 							}
 						}
diff --git a/Engine/Plugins/Experimental/Text3D/Source/Text3D/Public/Text3DComponent.h b/Engine/Plugins/Experimental/Text3D/Source/Text3D/Public/Text3DComponent.h
--- a/Engine/Plugins/Experimental/Text3D/Source/Text3D/Public/Text3DComponent.h
+++ b/Engine/Plugins/Experimental/Text3D/Source/Text3D/Public/Text3DComponent.h
@@ -456,7 +456,7 @@ protected:
 	UPROPERTY(VisibleInstanceOnly, Instanced, Category="Layout", meta=(AllowPrivateAccess="true"))
 	TObjectPtr<UText3DLayoutExtensionBase> LayoutExtension;
 
-	UPROPERTY(EditInstanceOnly, Instanced, Category="Layout|Effects", meta=(AllowPrivateAccess="true"))
+	UPROPERTY(EditInstanceOnly, Instanced, Category="LayoutEffects", meta=(AllowPrivateAccess="true"))
 	TArray<TObjectPtr<UText3DLayoutEffectBase>> LayoutEffects;
 
 	UPROPERTY(VisibleInstanceOnly, Instanced, Category="Geometry", meta=(AllowPrivateAccess="true"))
diff --git a/Engine/Plugins/Experimental/Text3D/Source/Text3DEditor/Private/Customizations/Text3DEditorTextComponentDetailCustomization.cpp b/Engine/Plugins/Experimental/Text3D/Source/Text3DEditor/Private/Customizations/Text3DEditorTextComponentDetailCustomization.cpp
--- a/Engine/Plugins/Experimental/Text3D/Source/Text3DEditor/Private/Customizations/Text3DEditorTextComponentDetailCustomization.cpp
+++ b/Engine/Plugins/Experimental/Text3D/Source/Text3DEditor/Private/Customizations/Text3DEditorTextComponentDetailCustomization.cpp
@@ -25,6 +25,7 @@ namespace UE::Text3DEditor::Customization
 			const FName LayoutSectionName = TEXT("Layout");
 			const TSharedRef<FPropertySection> LayoutSection = PropertyModule.FindOrCreateSection(ComponentClassName, LayoutSectionName, FText::FromName(LayoutSectionName));
 			LayoutSection->AddCategory(TEXT("Layout"));
+			LayoutSection->AddCategory(TEXT("LayoutEffects"));
 			LayoutSection->AddCategory(TEXT("Character"));
 
 			const FName RenderingSectionName = TEXT("Rendering");
@@ -38,6 +39,10 @@ namespace UE::Text3DEditor::Customization
 			const FName StyleSectionName = TEXT("Style");
 			const TSharedRef<FPropertySection> MaterialSection = PropertyModule.FindOrCreateSection(ComponentClassName, StyleSectionName, FText::FromName(StyleSectionName));
 			MaterialSection->AddCategory(TEXT("Material"));
+
+			const FName EffectSectionName = TEXT("Effects");
+			const TSharedRef<FPropertySection> EffectsSection = PropertyModule.FindOrCreateSection(ComponentClassName, EffectSectionName, FText::FromName(EffectSectionName));
+			EffectsSection->AddCategory(TEXT("LayoutEffects"));
 		}
 	}
 }
