MotionDesign : Text3D
- Fixed layout extension to handle RTL text
- Adjusted wrapping code for RTL text
- Fixed Text3D can be used in packaged builds

diff --git a/Engine/Plugins/Experimental/Text3D/Source/Text3D/Private/Characters/Text3DCharacterBase.cpp b/Engine/Plugins/Experimental/Text3D/Source/Text3D/Private/Characters/Text3DCharacterBase.cpp
--- a/Engine/Plugins/Experimental/Text3D/Source/Text3D/Private/Characters/Text3DCharacterBase.cpp
+++ b/Engine/Plugins/Experimental/Text3D/Source/Text3D/Private/Characters/Text3DCharacterBase.cpp
@@ -34,9 +34,34 @@ FTransform& UText3DCharacterBase::GetTransform(bool bInReset)
 	return Transform;
 }
 
-void UText3DCharacterBase::SetCharacterMetadata(const FText3DCharacterMetadata& InMetadata)
+void UText3DCharacterBase::SetGlyphIndex(uint32 InGlyphIndex)
 {
-	Metadata = InMetadata;
+	GlyphIndex = InGlyphIndex;
+}
+
+uint32 UText3DCharacterBase::GetGlyphIndex() const
+{
+	return GlyphIndex;
+}
+
+void UText3DCharacterBase::SetMeshBounds(const FBox& InBounds)
+{
+	MeshBounds = InBounds;
+}
+
+const FBox& UText3DCharacterBase::GetMeshBounds() const
+{
+	return MeshBounds;
+}
+
+void UText3DCharacterBase::SetMeshOffset(const FVector& InOffset)
+{
+	MeshOffset = InOffset;
+}
+
+const FVector& UText3DCharacterBase::GetMeshOffset() const
+{
+	return MeshOffset;
 }
 
 void UText3DCharacterBase::SetRelativeLocation(const FVector& InLocation)
@@ -94,7 +119,9 @@ void UText3DCharacterBase::ResetCharacterState()
 	RelativeScale = CDO->RelativeScale;
 	bVisible = CDO->bVisible;
 	Transform = CDO->Transform;
-	Metadata = CDO->Metadata;
+	GlyphIndex = CDO->GlyphIndex;
+	MeshBounds = CDO->MeshBounds;
+	MeshOffset = CDO->MeshOffset;
 }
 
 #if WITH_EDITOR
diff --git a/Engine/Plugins/Experimental/Text3D/Source/Text3D/Private/Extensions/Text3DDefaultCharacterExtension.cpp b/Engine/Plugins/Experimental/Text3D/Source/Text3D/Private/Extensions/Text3DDefaultCharacterExtension.cpp
--- a/Engine/Plugins/Experimental/Text3D/Source/Text3D/Private/Extensions/Text3DDefaultCharacterExtension.cpp
+++ b/Engine/Plugins/Experimental/Text3D/Source/Text3D/Private/Extensions/Text3DDefaultCharacterExtension.cpp
@@ -28,47 +28,18 @@ TConstArrayView<UText3DCharacterBase*> UText3DDefaultCharacterExtension::GetChar
 	return TextCharacters;
 }
 
-EText3DExtensionResult UText3DDefaultCharacterExtension::PreRendererUpdate(EText3DRendererFlags InFlag)
+void UText3DDefaultCharacterExtension::AllocateCharacters(uint16 InCount)
 {
-	if (InFlag != EText3DRendererFlags::Geometry)
-	{
-		return EText3DExtensionResult::Active;
-	}
-
-	const UText3DComponent* Text3DComponent = GetText3DComponent();
-	const FTypefaceEntry* TypeFaceEntry = Text3DComponent->GetTypeFaceEntry();
-
-	if (!TypeFaceEntry)
-	{
-		UE_LOG(LogText3D, Warning, TEXT("Invalid typeface entry cannot proceed"))
-		return EText3DExtensionResult::Failed;
-	}
-
-	const FString FormattedTextStr = Text3DComponent->GetFormattedText().ToString();
-
-	const TArray<FText3DCharacterMetadata> TextCharacterMetadata = UE::Text3D::Utilities::GetCharacterMetadata(*TypeFaceEntry, FormattedTextStr);
-
-	AllocateTextCharacters(TextCharacterMetadata.Num());
-
-	for (TConstEnumerateRef<TObjectPtr<UText3DCharacterBase>> Character : EnumerateRange(TextCharacters))
-	{
-		const FText3DCharacterMetadata& Metadata = TextCharacterMetadata[Character.GetIndex()];
-#if WITH_EDITORONLY_DATA
-		(*Character)->SetCharacter(FString::Printf(TEXT("%c"), FormattedTextStr[Metadata.TextIndex]));
-#endif
-		(*Character)->SetCharacterMetadata(Metadata);
-	}
-
-	return EText3DExtensionResult::Finished;
-}
-
-EText3DExtensionResult UText3DDefaultCharacterExtension::PostRendererUpdate(EText3DRendererFlags InFlag)
-{
-	return EText3DExtensionResult::Active;
+	AllocateTextCharacters(InCount);
 }
 
 void UText3DDefaultCharacterExtension::AllocateTextCharacters(uint16 InCharacterCount)
 {
+	if (TextCharacters.Num() == InCharacterCount)
+	{
+		return;
+	}
+
 	TextCharacters.Reserve(InCharacterCount);
 
 	const int32 RemainingCharacterCount = TextCharacters.Num() - InCharacterCount;
diff --git a/Engine/Plugins/Experimental/Text3D/Source/Text3D/Private/Extensions/Text3DDefaultGeometryExtension.cpp b/Engine/Plugins/Experimental/Text3D/Source/Text3D/Private/Extensions/Text3DDefaultGeometryExtension.cpp
--- a/Engine/Plugins/Experimental/Text3D/Source/Text3D/Private/Extensions/Text3DDefaultGeometryExtension.cpp
+++ b/Engine/Plugins/Experimental/Text3D/Source/Text3D/Private/Extensions/Text3DDefaultGeometryExtension.cpp
@@ -3,7 +3,9 @@
 #include "Extensions/Text3DDefaultGeometryExtension.h"
 
 #include "Characters/Text3DCharacterBase.h"
+#include "Engine/Font.h"
 #include "Engine/StaticMesh.h"
+#include "Logs/Text3DLogs.h"
 #include "MeshDescription.h"
 #include "Subsystems/Text3DEngineSubsystem.h"
 #include "Text3DComponent.h"
@@ -135,13 +137,12 @@ void UText3DDefaultGeometryExtension::PostEditChangeProperty(FPropertyChangedEve
 }
 #endif
 
-UStaticMesh* UText3DDefaultGeometryExtension::GetGlyphMesh(UText3DCharacterBase* InCharacter) const
+UStaticMesh* UText3DDefaultGeometryExtension::FindOrLoadGlyphMesh(UText3DCharacterBase* InCharacter) const
 {
-	const uint16 GlyphIndex = InCharacter->GetCharacterMetadata().GlyphIndex;
-	const FText3DCachedMesh* GlyphMesh = CachedFontData->GetGlyphMesh(GlyphIndex, GlyphMeshParameters);
+	const FText3DCachedMesh* GlyphMesh = CachedFontData->GetGlyphMesh(InCharacter->GetGlyphIndex(), GlyphMeshParameters);
 	checkf(GlyphMesh != nullptr, TEXT("Invalid glyph mesh returned"));
-	InCharacter->GetCharacterMetadata().MeshBounds = GlyphMesh->MeshBounds;
-	InCharacter->GetCharacterMetadata().MeshOffset = GlyphMesh->MeshOffset;
+	InCharacter->SetMeshBounds(GlyphMesh->MeshBounds);
+	InCharacter->SetMeshOffset(GlyphMesh->MeshOffset);
 	return GlyphMesh->Mesh;
 }
 
@@ -155,6 +156,16 @@ EText3DVerticalTextAlignment UText3DDefaultGeometryExtension::GetGlyphVAlignment
 	return PivotVAlignment;
 }
 
+FTypefaceFontData* UText3DDefaultGeometryExtension::GetTypefaceFontData() const
+{
+	if (CachedFontData)
+	{
+		return CachedFontData->GetTypefaceFontData(GlyphMeshParameters.TypefaceIndex);
+	}
+
+	return nullptr;
+}
+
 EText3DExtensionResult UText3DDefaultGeometryExtension::PreRendererUpdate(EText3DRendererFlags InFlag)
 {
 	if (InFlag != EText3DRendererFlags::Geometry)
@@ -165,11 +176,13 @@ EText3DExtensionResult UText3DDefaultGeometryExtension::PreRendererUpdate(EText3
 	const UText3DComponent* TextComponent = GetText3DComponent();
 	UText3DEngineSubsystem* Subsystem = UText3DEngineSubsystem::Get();
 	const uint32 TypefaceIndex = TextComponent->GetTypeFaceIndex();
+	UFont* Font = TextComponent->GetFont();
 
-	CachedFontData = &Subsystem->GetCachedFontData(TextComponent->GetFont(), TypefaceIndex);
+	CachedFontData = &Subsystem->GetCachedFontData(Font, TypefaceIndex);
 
 	if (!CachedFontData)
 	{
+		UE_LOG(LogText3D, Error, TEXT("Failed to get cached font data for '%s %i' in Text3D geometry extension"), Font ? *Font->GetName() : TEXT("invalid"), TypefaceIndex);
 		return EText3DExtensionResult::Failed;
 	}
 
diff --git a/Engine/Plugins/Experimental/Text3D/Source/Text3D/Private/Extensions/Text3DDefaultLayoutExtension.cpp b/Engine/Plugins/Experimental/Text3D/Source/Text3D/Private/Extensions/Text3DDefaultLayoutExtension.cpp
--- a/Engine/Plugins/Experimental/Text3D/Source/Text3D/Private/Extensions/Text3DDefaultLayoutExtension.cpp
+++ b/Engine/Plugins/Experimental/Text3D/Source/Text3D/Private/Extensions/Text3DDefaultLayoutExtension.cpp
@@ -3,22 +3,21 @@
 #include "Extensions/Text3DDefaultLayoutExtension.h"
 
 #include "Algo/Accumulate.h"
+#include "Algo/Count.h"
 #include "Characters/Text3DCharacterBase.h"
 #include "Engine/Font.h"
+#include "Extensions/Text3DCharacterExtensionBase.h"
 #include "Extensions/Text3DGeometryExtensionBase.h"
-#include "Extensions/Text3DLayoutEffectBase.h"
 #include "Fonts/SlateTextShaper.h"
 #include "Framework/Text/PlainTextLayoutMarshaller.h"
 #include "LayoutBuilders/Text3DLayout.h"
 #include "LayoutBuilders/Text3DLayoutShaper.h"
 #include "LayoutBuilders/Text3DShapedGlyphText.h"
-#include "Logs/Text3DLogs.h"
 #include "Misc/EnumerateRange.h"
 #include "Styling/StyleDefaults.h"
 #include "Styling/SlateTypes.h"
 #include "Subsystems/Text3DEngineSubsystem.h"
 #include "Text3DComponent.h"
-#include "Utilities/Text3DUtilities.h"
 
 FName UText3DDefaultLayoutExtension::GetUseMaxWidthPropertyName()
 {
@@ -180,7 +179,7 @@ FVector UText3DDefaultLayoutExtension::GetTextScale() const
 
 EText3DExtensionResult UText3DDefaultLayoutExtension::PreRendererUpdate(EText3DRendererFlags InFlag)
 {
-	if (InFlag != EText3DRendererFlags::Layout)
+	if (InFlag != EText3DRendererFlags::Geometry && InFlag != EText3DRendererFlags::Layout)
 	{
 		return EText3DExtensionResult::Active;
 	}
@@ -209,76 +208,87 @@ EText3DExtensionResult UText3DDefaultLayoutExtension::PreRendererUpdate(EText3DR
 
 	if (EnumHasAnyFlags(Text3DComponent->GetUpdateFlags(), EText3DRendererFlags::Geometry))
 	{
-		FSlateFontInfo FontInfo(Text3DComponent->GetFont(), UE::Text3D::Metrics::SlateFontSize);
-		FontInfo.CompositeFont = FStyleDefaults::GetFontInfo().CompositeFont;
+		UFont* Font = Text3DComponent->GetFont();
+		FSlateFontInfo FontInfo(Font, UE::Text3D::Metrics::SlateFontSize);
+		FontInfo.CompositeFont = Font->GetCompositeFont() ? MakeShared<FCompositeFont>(*Text3DComponent->GetFont()->GetCompositeFont()) : FStyleDefaults::GetFontInfo().CompositeFont;
 		FontInfo.TypefaceFontName = Text3DComponent->GetTypeface();
 		FontStyle->SetFont(FontInfo);
-
-		UText3DEngineSubsystem* Subsystem = UText3DEngineSubsystem::Get();
-
-		const uint32 TypefaceIndex = Text3DComponent->GetTypeFaceIndex();
-		FCachedFontData& CachedFontData = Subsystem->GetCachedFontData(Text3DComponent->GetFont(), TypefaceIndex);
-		const FT_Face Face = CachedFontData.GetFreeTypeFace(TypefaceIndex);
-
-		if (!Face)
-		{
-			UE_LOG(LogText3D, Error, TEXT("Failed to load font face data '%s' in Text3D layout extension"), *CachedFontData.GetFontName());
-			return EText3DExtensionResult::Failed;
-		}
-
-		TypeFaceMetrics.FontHeight = Face->size->metrics.height * UE::Text3D::Metrics::FontSizeInverse;
-		TypeFaceMetrics.FontAscender = Face->size->metrics.ascender * UE::Text3D::Metrics::FontSizeInverse;
-		TypeFaceMetrics.FontDescender = Face->size->metrics.descender * UE::Text3D::Metrics::FontSizeInverse;
 	}
 
 	ShapedText->Reset();
 	FText3DLayoutShaper::Get()->ShapeBidirectionalText(*FontStyle, Text3DComponent->GetFormattedText().ToString(), TextLayout, TextLayoutMarshaller, ShapedText->Lines);
 
-	const int32 ApproximateGlyphNum = Algo::TransformAccumulate(ShapedText->Lines,
-		[](const FShapedGlyphLine& InLine)
+	const int32 CharacterCount = Algo::TransformAccumulate(ShapedText->Lines, [](const FShapedGlyphLine& InLine)
+	{
+		return Algo::CountIf(InLine.GlyphsToRender, [](const FShapedGlyphEntry& InGlyph)
 		{
-			return InLine.GlyphsToRender.Num();
-		},
-		0);
+			return InGlyph.HasValidGlyph() && InGlyph.bIsVisible;
+		});
+	},0);
+
+	if (EnumHasAnyFlags(Text3DComponent->GetUpdateFlags(), EText3DRendererFlags::Geometry))
+	{
+		UText3DCharacterExtensionBase* CharacterExtension = Text3DComponent->GetCharacterExtension();
+		CharacterExtension->AllocateCharacters(CharacterCount);
+	}
+
+	ShapedText->GlyphIndexToFontFace.Reserve(CharacterCount);
+	ShapedText->Kernings.SetNum(CharacterCount);
 
-	ShapedText->GlyphIndexToFontFace.Reserve(ApproximateGlyphNum);
-	for (int32 LineIndex = 0; LineIndex < ShapedText->Lines.Num(); ++LineIndex)
+	uint16 CharacterIndex = 0;
+	FT_Face FontFace = nullptr;
+	for (TConstEnumerateRef<FShapedGlyphLine> ShapedLine : EnumerateRange(ShapedText->Lines))
 	{
-		for (const FShapedGlyphEntry& GlyphEntry : ShapedText->Lines[LineIndex].GlyphsToRender)
+		const bool bIsRTL = (ShapedLine->TextDirection == TextBiDi::ETextDirection::RightToLeft);
+
+		for (TConstEnumerateRef<FShapedGlyphEntry> GlyphEntry : EnumerateRange(ShapedLine->GlyphsToRender))
 		{
-			if (!GlyphEntry.FontFaceData.IsValid())
+			if (!GlyphEntry->FontFaceData.IsValid())
 			{
 				// Add as nullptr if not already in map
-				if (!ShapedText->GlyphIndexToFontFace.Contains(GlyphEntry.GlyphIndex))
+				if (!ShapedText->GlyphIndexToFontFace.Contains(GlyphEntry->GlyphIndex))
 				{
-					ShapedText->GlyphIndexToFontFace.Add(GlyphEntry.GlyphIndex, nullptr);
+					ShapedText->GlyphIndexToFontFace.Add(GlyphEntry->GlyphIndex, nullptr);
 				}
 
 				continue;
 			}
 
-			const TSharedPtr<FFreeTypeFace> FontFacePtr = GlyphEntry.FontFaceData->FontFace.Pin();
+			if (const TSharedPtr<FFreeTypeFace> FontFacePtr = GlyphEntry->FontFaceData->FontFace.Pin())
+			{
+				ShapedText->GlyphIndexToFontFace.FindOrAdd(GlyphEntry->GlyphIndex, FontFacePtr.Get());
+				FontFace = FontFacePtr->GetFace();
+			}
 
-			if (FontFacePtr.IsValid())
+			if (GlyphEntry->HasValidGlyph() && GlyphEntry->bIsVisible)
 			{
-				ShapedText->GlyphIndexToFontFace.FindOrAdd(GlyphEntry.GlyphIndex, FontFacePtr.Get());
+				const uint16 Index = bIsRTL ? CharacterCount - 1 - CharacterIndex : CharacterIndex;
+				UText3DCharacterBase* Character = Text3DComponent->GetCharacter(Index);
+				ShapedText->Kernings[CharacterIndex] = Character->GetCharacterKerning();
+				CharacterIndex++;
 			}
 		}
 	}
 
-	ShapedText->CharacterGlyphMap.Reset();
-	ShapedText->Kernings.SetNum(Text3DComponent->GetCharacterCount());
-	ShapedText->GlyphSizes.SetNum(Text3DComponent->GetCharacterCount());
-	Text3DComponent->ForEachCharacter([this](UText3DCharacterBase* InCharacter, uint16 InIndex, uint16 InCount)
+	UText3DGeometryExtensionBase* GeometryExtension = Text3DComponent->GetGeometryExtension();
+	if (EnumHasAnyFlags(Text3DComponent->GetUpdateFlags(), EText3DRendererFlags::Geometry))
 	{
-		ShapedText->Kernings[InIndex] = InCharacter->GetCharacterKerning();
-		ShapedText->GlyphSizes[InIndex] = InCharacter->GetCharacterMetadata().MeshBounds.GetSize();
-		ShapedText->CharacterGlyphMap.Add(InCharacter->GetCharacterMetadata().LayoutIndex, InIndex);
-	});
+		FTypefaceFontData* TypefaceFontData = GeometryExtension->GetTypefaceFontData();
+		if (!TypefaceFontData)
+		{
+			return EText3DExtensionResult::Failed;
+		}
+
+		TypefaceFontData->SetTypeFace(FontFace);
+
+		if (FontFace)
+		{
+			TypeFaceMetrics.FontHeight = FontFace->size->metrics.height * UE::Text3D::Metrics::FontSizeInverse;
+			TypeFaceMetrics.FontAscender = FontFace->size->metrics.ascender * UE::Text3D::Metrics::FontSizeInverse;
+			TypeFaceMetrics.FontDescender = FontFace->size->metrics.descender * UE::Text3D::Metrics::FontSizeInverse;
+		}
+	}
 
-	UText3DGeometryExtensionBase* GeometryExtension = Text3DComponent->GetGeometryExtension();
-	ShapedText->PivotHAlignment = GeometryExtension->GetGlyphHAlignment();
-	ShapedText->PivotVAlignment = GeometryExtension->GetGlyphVAlignment();
 	ShapedText->Tracking = Tracking;
 	ShapedText->WordSpacing = WordSpacing;
 	ShapedText->MaxWidth = MaxWidth;
@@ -286,23 +296,31 @@ EText3DExtensionResult UText3DDefaultLayoutExtension::PreRendererUpdate(EText3DR
 	ShapedText->CalculateWidth();
 	CalculateTextScale();
 
-	int32 LayoutIndex = 0;
+	CharacterIndex = 0;
 	for (TConstEnumerateRef<FShapedGlyphLine> ShapedLine : EnumerateRange(ShapedText->Lines))
 	{
 		FVector Location = GetLineLocation(ShapedLine.GetIndex());
+		const bool bIsRTL = (ShapedLine->TextDirection == TextBiDi::ETextDirection::RightToLeft);
 
 		for (TConstEnumerateRef<FShapedGlyphEntry> ShapedGlyph : EnumerateRange(ShapedLine->GlyphsToRender))
 		{
-			if (const int32* CharacterIndex = ShapedText->CharacterGlyphMap.Find(LayoutIndex))
+			if (ShapedGlyph->HasValidGlyph() && ShapedGlyph->bIsVisible)
 			{
-				UText3DCharacterBase* CharacterData = Text3DComponent->GetCharacter(*CharacterIndex);
+				const uint16 Index = bIsRTL ? CharacterCount - 1 - CharacterIndex : CharacterIndex;
+				UText3DCharacterBase* Character = Text3DComponent->GetCharacter(Index);
+
+				if (EnumHasAnyFlags(Text3DComponent->GetUpdateFlags(), EText3DRendererFlags::Geometry))
+				{
+					Character->SetGlyphIndex(ShapedGlyph->GlyphIndex);
+					GeometryExtension->FindOrLoadGlyphMesh(Character);
+				}
 
 				constexpr bool bReset = true;
-				FTransform& CharacterTransform = CharacterData->GetTransform(bReset);
+				FTransform& CharacterTransform = Character->GetTransform(bReset);
 
-				const FVector GlyphSize = ShapedText->GlyphSizes[*CharacterIndex];
+				const FVector GlyphSize = Character->GetMeshBounds().GetSize();
 				FVector Shift = FVector::ZeroVector;
-				switch (ShapedText->PivotHAlignment)
+				switch (GeometryExtension->GetGlyphHAlignment())
 				{
 				case EText3DHorizontalTextAlignment::Left:
 					Shift = FVector::ZeroVector;
@@ -317,14 +335,14 @@ EText3DExtensionResult UText3DDefaultLayoutExtension::PreRendererUpdate(EText3DR
 
 				if (ShapedGlyph.GetIndex() != 0)
 				{
-					Shift.Y += Tracking + ShapedText->Kernings[*CharacterIndex];
+					Shift.Y += Tracking + ShapedText->Kernings[CharacterIndex];
 				}
 
 				CharacterTransform.Accumulate(FTransform(Location + Shift));
+				CharacterIndex++;
 			}
 
 			Location.Y += ShapedLine->GlyphAdvances[ShapedGlyph.GetIndex()];
-			LayoutIndex++;
 		}
 	}
 
@@ -447,5 +465,18 @@ void UText3DDefaultLayoutExtension::OnLayoutOptionsChanged()
 
 	EText3DRendererFlags Flags = EText3DRendererFlags::Material;
 	EnumAddFlags(Flags, EText3DRendererFlags::Layout);
+
+	if (MaxWidthBehavior == EText3DMaxWidthHandling::WrapAndScale)
+	{
+		for (const TConstEnumerateRef<FShapedGlyphLine> ShapedLine : EnumerateRange(ShapedText->Lines))
+		{
+			if (ShapedLine->TextDirection == TextBiDi::ETextDirection::RightToLeft)
+			{
+				EnumAddFlags(Flags, EText3DRendererFlags::Geometry);
+				break;
+			}
+		}
+	}
+
 	RequestUpdate(Flags);
 }
diff --git a/Engine/Plugins/Experimental/Text3D/Source/Text3D/Private/LayoutBuilders/Text3DLayoutShaper.cpp b/Engine/Plugins/Experimental/Text3D/Source/Text3D/Private/LayoutBuilders/Text3DLayoutShaper.cpp
--- a/Engine/Plugins/Experimental/Text3D/Source/Text3D/Private/LayoutBuilders/Text3DLayoutShaper.cpp
+++ b/Engine/Plugins/Experimental/Text3D/Source/Text3D/Private/LayoutBuilders/Text3DLayoutShaper.cpp
@@ -49,6 +49,8 @@ void FText3DLayoutShaper::ShapeBidirectionalText(
 	for (const FTextLayout::FLineView& Line : LineViews)
 	{
 		FShapedGlyphLine& ShapedLine = OutShapedLines.AddDefaulted_GetRef();
+		ShapedLine.TextDirection = Line.TextBaseDirection;
+
 		for (const TSharedRef<ILayoutBlock>& Block : Line.Blocks)
 		{
 			FString BlockText;
@@ -59,7 +61,7 @@ void FText3DLayoutShaper::ShapeBidirectionalText(
 			const FTextRange BlockRange = Block->GetTextRange();
 			const FLayoutBlockTextContext BlockTextContext = Block->GetTextContext();
 
-			FShapedGlyphSequenceRef ShapedGlyphSequence = ShapedTextCacheUtil::GetShapedTextSubSequence(
+			const FShapedGlyphSequenceRef ShapedGlyphSequence = ShapedTextCacheUtil::GetShapedTextSubSequence(
 				BlockTextContext.ShapedTextCache,
 				FCachedShapedTextKey(Line.Range, 1.0f, BlockTextContext, Style.Font),
 				BlockRange,
@@ -69,5 +71,10 @@ void FText3DLayoutShaper::ShapeBidirectionalText(
 
 			ShapedLine.GlyphsToRender.Append(ShapedGlyphSequence->GetGlyphsToRender());
 		}
+
+		if (ShapedLine.TextDirection == TextBiDi::ETextDirection::RightToLeft)
+		{
+			Algo::Reverse(ShapedLine.GlyphsToRender);
+		}
 	}
 }
diff --git a/Engine/Plugins/Experimental/Text3D/Source/Text3D/Private/LayoutBuilders/Text3DShapedGlyphLine.cpp b/Engine/Plugins/Experimental/Text3D/Source/Text3D/Private/LayoutBuilders/Text3DShapedGlyphLine.cpp
--- a/Engine/Plugins/Experimental/Text3D/Source/Text3D/Private/LayoutBuilders/Text3DShapedGlyphLine.cpp
+++ b/Engine/Plugins/Experimental/Text3D/Source/Text3D/Private/LayoutBuilders/Text3DShapedGlyphLine.cpp
@@ -11,9 +11,9 @@ float FShapedGlyphLine::GetWidthAdvance(const int32 InIndex, const float InKerni
 
 	const bool bFirstCharacter = InIndex == 0;
 	const bool bLastCharacter = InIndex == GlyphsToRender.Num() - 1;
+	const bool bSkipCharacter = TextDirection == TextBiDi::ETextDirection::RightToLeft ? bLastCharacter : bFirstCharacter;
 
-	// Skip first and last character for kerning and tracking
-	if (!bFirstCharacter && !bLastCharacter)
+	if (!bSkipCharacter)
 	{
 		// @note: as per FSlateElementBatcher::BuildShapedTextSequence, per Glyph Kerning isn't used
 		Advance += InKerning;
diff --git a/Engine/Plugins/Experimental/Text3D/Source/Text3D/Private/LayoutBuilders/Text3DShapedGlyphLine.h b/Engine/Plugins/Experimental/Text3D/Source/Text3D/Private/LayoutBuilders/Text3DShapedGlyphLine.h
--- a/Engine/Plugins/Experimental/Text3D/Source/Text3D/Private/LayoutBuilders/Text3DShapedGlyphLine.h
+++ b/Engine/Plugins/Experimental/Text3D/Source/Text3D/Private/LayoutBuilders/Text3DShapedGlyphLine.h
@@ -20,4 +20,7 @@ struct FShapedGlyphLine
 
 	/** Glyph advance on this line */
 	TArray<float> GlyphAdvances;
-};
\ No newline at end of file
+
+	/** Direction of the text in the line */
+	TextBiDi::ETextDirection TextDirection = TextBiDi::ETextDirection::LeftToRight;
+};
diff --git a/Engine/Plugins/Experimental/Text3D/Source/Text3D/Private/LayoutBuilders/Text3DShapedGlyphText.cpp b/Engine/Plugins/Experimental/Text3D/Source/Text3D/Private/LayoutBuilders/Text3DShapedGlyphText.cpp
--- a/Engine/Plugins/Experimental/Text3D/Source/Text3D/Private/LayoutBuilders/Text3DShapedGlyphText.cpp
+++ b/Engine/Plugins/Experimental/Text3D/Source/Text3D/Private/LayoutBuilders/Text3DShapedGlyphText.cpp
@@ -11,71 +11,95 @@ FText3DShapedGlyphText::FText3DShapedGlyphText()
 
 void FText3DShapedGlyphText::Reset()
 {
-	PivotHAlignment = EText3DHorizontalTextAlignment::Left;
-	PivotVAlignment = EText3DVerticalTextAlignment::Bottom;
 	Kernings.Reset();
-	GlyphSizes.Reset();
 	Tracking = 0.f;
 	WordSpacing = 0.0f;
 	bWrap = false;
 	Lines.Reset();
 	GlyphIndexToFontFace.Reset();
-	CharacterGlyphMap.Reset();
 }
 
 void FText3DShapedGlyphText::CalculateWidth()
 {
 	TArray<FShapedGlyphLine> NewLines;
-	NewLines.Reserve(Lines.Num());
+    NewLines.Reserve(Lines.Num());
 
-	int32 LayoutIndex = 0;
-	for (const TConstEnumerateRef<FShapedGlyphLine> GlyphLine : EnumerateRange(Lines))
-	{
-		FShapedGlyphLine* CurrentLine = &NewLines.AddDefaulted_GetRef();
+    int32 CharacterIndex = 0;
+    for (TConstEnumerateRef<FShapedGlyphLine> GlyphLine : EnumerateRange(Lines))
+    {
+        FShapedGlyphLine* CurrentLine = &NewLines.AddDefaulted_GetRef();
+        CurrentLine->TextDirection = GlyphLine->TextDirection;
 
-		TArray<float> CurrentAdvances;
-		TArray<FShapedGlyphEntry> CurrentWord;
-		float LineWidth = 0.0f;
-		float CurrentWordLength = 0.0f;
+        const bool bIsRTL = (GlyphLine->TextDirection == TextBiDi::ETextDirection::RightToLeft);
+
+        TArray<float> CurrentAdvances;
+        TArray<FShapedGlyphEntry> CurrentWord;
+        float LineWidth = 0.0f;
+        float CurrentWordLength = 0.0f;
+
+        const int32 GlyphCount = GlyphLine->GlyphsToRender.Num();
 
 		for (TConstEnumerateRef<FShapedGlyphEntry> CurrentGlyph : EnumerateRange(GlyphLine->GlyphsToRender))
-		{
-			const int32* CharacterIndex = CharacterGlyphMap.Find(LayoutIndex);
-			const float CurrentKerning = !!CharacterIndex ? Kernings[*CharacterIndex] : 0.f;
-			const float GlyphAdv = GlyphLine->GetWidthAdvance(CurrentGlyph.GetIndex(), Tracking + CurrentKerning, WordSpacing);
-
-			// If we're at the end the line or at whitespace
-			const bool bWordBreak = !CurrentGlyph->bIsVisible || CurrentGlyph.GetIndex() == GlyphLine->GlyphsToRender.Num() - 1;
-			if (bWrap                               // when we're wrapping
+        {
+    		const bool bVisibleGlyph = CurrentGlyph->HasValidGlyph() && CurrentGlyph->bIsVisible;
+			const int32 Index = bIsRTL ? Kernings.Num() - 1 - CharacterIndex : CharacterIndex;
+    		const float CurrentKerning = bVisibleGlyph ? Kernings[Index] : 0.f;
+    		const float GlyphAdv = GlyphLine->GetWidthAdvance(CurrentGlyph.GetIndex(), Tracking + CurrentKerning, WordSpacing);
+
+            // If we're at the end the line or at whitespace
+            const bool bWordBreak = !bVisibleGlyph || CurrentGlyph.GetIndex() == GlyphCount - 1;
+
+    		if (bWrap                               // when we're wrapping
 				&& bWordBreak                       // and at a word break
 				&& LineWidth > MaxWidth             // and the current line is longer than the max
 				&& CurrentWordLength != LineWidth)  // and the line is not just a single word that we can't break
-			{
-				CurrentLine->Width = LineWidth - CurrentWordLength;
-				CurrentLine = &NewLines.AddDefaulted_GetRef();
-				LineWidth = CurrentWordLength;
-			}
+    		{
+    			CurrentLine->Width = LineWidth - CurrentWordLength;
+    			CurrentLine = &NewLines.AddDefaulted_GetRef();
+    			CurrentLine->TextDirection = GlyphLine->TextDirection;
+    			LineWidth = CurrentWordLength;
+    		}
 
-			CurrentWord.Add(*CurrentGlyph);
 			LineWidth += GlyphAdv;
 			CurrentWordLength += GlyphAdv;
-			CurrentAdvances.Add(GlyphAdv);
 
-			if (bWordBreak)
+			if (bIsRTL)
 			{
-				CurrentLine->GlyphsToRender.Append(CurrentWord);
-				CurrentLine->GlyphAdvances.Append(CurrentAdvances);
-				CurrentWordLength = 0.0f;
-				CurrentWord.Reset();
-				CurrentAdvances.Reset();
+				CurrentWord.Insert(*CurrentGlyph, 0);
+				CurrentAdvances.Insert(GlyphAdv, 0);
+			}
+			else
+			{
+				CurrentWord.Add(*CurrentGlyph);
+				CurrentAdvances.Add(GlyphAdv);
 			}
 
-			LayoutIndex++;
-		}
+            if (bWordBreak)
+            {
+        		if (bIsRTL)
+        		{
+        			CurrentLine->GlyphsToRender.Insert(CurrentWord, 0);
+        			CurrentLine->GlyphAdvances.Insert(CurrentAdvances, 0);
+        		}
+	            else
+	            {
+            		CurrentLine->GlyphsToRender.Append(CurrentWord);
+            		CurrentLine->GlyphAdvances.Append(CurrentAdvances);
+	            }
 
-		CurrentLine->Width = LineWidth;
-	}
-	
-	Lines = NewLines;
-}
+                CurrentWordLength = 0.0f;
+                CurrentWord.Reset();
+                CurrentAdvances.Reset();
+            }
+
+    		if (bVisibleGlyph)
+    		{
+    			CharacterIndex++;
+    		}
+        }
+
+        CurrentLine->Width = LineWidth;
+    }
 
+    Lines = NewLines;
+}
\ No newline at end of file
diff --git a/Engine/Plugins/Experimental/Text3D/Source/Text3D/Private/LayoutBuilders/Text3DShapedGlyphText.h b/Engine/Plugins/Experimental/Text3D/Source/Text3D/Private/LayoutBuilders/Text3DShapedGlyphText.h
--- a/Engine/Plugins/Experimental/Text3D/Source/Text3D/Private/LayoutBuilders/Text3DShapedGlyphText.h
+++ b/Engine/Plugins/Experimental/Text3D/Source/Text3D/Private/LayoutBuilders/Text3DShapedGlyphText.h
@@ -7,6 +7,7 @@
 #include "Text3DShapedGlyphLine.h"
 
 class FFreeTypeFace;
+struct FShapedGlyphLine;
 
 /** Contains text metrics with sufficient information to fetch and transform each line. */
 struct FText3DShapedGlyphText
@@ -16,21 +17,9 @@ struct FText3DShapedGlyphText
 	void Reset();
 	void CalculateWidth();
 
-	/** Maps character layout index to character index */
-	TMap<int32, int32> CharacterGlyphMap;
-
-	/** Pivot Horizontal alignment */
-	EText3DHorizontalTextAlignment PivotHAlignment;
-
-	/** Pivot Vertical alignment */
-	EText3DVerticalTextAlignment PivotVAlignment;
-
 	/** Individual character kerning */
 	TArray<float> Kernings;
 
-	/** Individual character size */
-	TArray<FVector> GlyphSizes;
-
 	/** General tracking for all characters */
 	float Tracking;
 
diff --git a/Engine/Plugins/Experimental/Text3D/Source/Text3D/Private/Renderers/StaticMeshes/Text3DStaticMeshesRenderer.cpp b/Engine/Plugins/Experimental/Text3D/Source/Text3D/Private/Renderers/StaticMeshes/Text3DStaticMeshesRenderer.cpp
--- a/Engine/Plugins/Experimental/Text3D/Source/Text3D/Private/Renderers/StaticMeshes/Text3DStaticMeshesRenderer.cpp
+++ b/Engine/Plugins/Experimental/Text3D/Source/Text3D/Private/Renderers/StaticMeshes/Text3DStaticMeshesRenderer.cpp
@@ -42,7 +42,7 @@ void UText3DStaticMeshesRenderer::OnUpdate(EText3DRendererFlags InFlags)
 
 		TextComponent->ForEachCharacter([this, GeometryExtension](UText3DCharacterBase* InCharacter, uint16 InIndex, uint16 InCount)
 		{
-			UStaticMesh* CachedMesh = GeometryExtension->GetGlyphMesh(InCharacter);
+			UStaticMesh* CachedMesh = GeometryExtension->FindOrLoadGlyphMesh(InCharacter);
 			CharacterMeshes[InIndex]->SetStaticMesh(CachedMesh);
 		});
 
diff --git a/Engine/Plugins/Experimental/Text3D/Source/Text3D/Private/Subsystems/Text3DEngineSubsystem.cpp b/Engine/Plugins/Experimental/Text3D/Source/Text3D/Private/Subsystems/Text3DEngineSubsystem.cpp
--- a/Engine/Plugins/Experimental/Text3D/Source/Text3D/Private/Subsystems/Text3DEngineSubsystem.cpp
+++ b/Engine/Plugins/Experimental/Text3D/Source/Text3D/Private/Subsystems/Text3DEngineSubsystem.cpp
@@ -183,12 +183,25 @@ FCachedFontData& UText3DEngineSubsystem::GetCachedFontData(UFont* InFont, int32
 }
 
 FTypefaceFontData::FTypefaceFontData()
-	: FreeTypeFace(nullptr)
-	, TypefaceFontDataHash(0)
 {
 	CacheCounter = MakeShared<int32>(0);
 }
 
+void FTypefaceFontData::SetTypeFace(FT_Face InTypeFace)
+{
+	TypeFace = InTypeFace;
+}
+
+TSharedPtr<const FFontFaceData> FTypefaceFontData::GetTypeFaceData() const
+{
+	return TypeFaceDataWeak.Pin();
+}
+
+void FTypefaceFontData::SetTypeFaceData(const TSharedPtr<const FFontFaceData>& InTypeFaceData)
+{
+	TypeFaceDataWeak = InTypeFaceData;
+}
+
 FCachedFontData::FCachedFontData()
 {
 	Font = nullptr;
@@ -206,12 +219,6 @@ void FCachedFontData::ClearFreeTypeFace()
 
 void FCachedFontData::LoadFreeTypeFace(uint32 InTypefaceIndex)
 {
-	const bool bContainsTypeface = TypefaceFontDataMap.Contains(InTypefaceIndex);
-	if (bContainsTypeface)
-	{
-		return;
-	}
-
 	if (!Font)
 	{
 		return;
@@ -231,28 +238,10 @@ void FCachedFontData::LoadFreeTypeFace(uint32 InTypefaceIndex)
 	const FTypefaceEntry& Typeface = CompositeFont->DefaultTypeface.Fonts[InTypefaceIndex];
 	const FFontFaceDataConstPtr FaceData = Typeface.Font.GetFontFaceData();
 
-	if (!bContainsTypeface)
-	{
-		FTypefaceFontData& TypefaceFontData = TypefaceFontDataMap.Add(InTypefaceIndex);
-		TypefaceFontData.SetTypefaceName(Typeface.Name);
-	}
-
-	if (FaceData.IsValid() && FaceData->HasData() && FaceData->GetData().Num() > 0)
-	{
-		TypefaceFontDataMap[InTypefaceIndex].SetTypefaceData(FaceData->GetData());
-		FT_New_Memory_Face(FText3DModule::GetFreeTypeLibrary(), TypefaceFontDataMap[InTypefaceIndex].GetTypefaceData().GetData(), TypefaceFontDataMap[InTypefaceIndex].GetTypefaceData().Num(), 0, &TypefaceFontDataMap[InTypefaceIndex].GetTypeface());
-	}
-	else if (FFileHelper::LoadFileToArray(TypefaceFontDataMap[InTypefaceIndex].GetTypefaceData(), *Typeface.Font.GetFontFilename()) && TypefaceFontDataMap[InTypefaceIndex].GetTypefaceData().Num() > 0)
-	{
-		FT_New_Memory_Face(FText3DModule::GetFreeTypeLibrary(), TypefaceFontDataMap[InTypefaceIndex].GetTypefaceData().GetData(), TypefaceFontDataMap[InTypefaceIndex].GetTypefaceData().Num(), 0, &TypefaceFontDataMap[InTypefaceIndex].GetTypeface());
-	}
-
-	if (const FT_Face* FreeTypeface = &TypefaceFontDataMap[InTypefaceIndex].GetTypeface())
-	{
-		TypefaceFontDataMap[InTypefaceIndex].SetTypefaceFontDataHash(HashCombine(0, GetTypeHash(Typeface.Font)));
-		FT_Set_Char_Size(*FreeTypeface, UE::Text3D::Metrics::FontSize, UE::Text3D::Metrics::FontSize, UE::Text3D::Metrics::FontDPI, UE::Text3D::Metrics::FontDPI);
-		FT_Set_Pixel_Sizes(*FreeTypeface, UE::Text3D::Metrics::FontSize, UE::Text3D::Metrics::FontSize);
-	}
+	FTypefaceFontData& TypefaceFontData = TypefaceFontDataMap.FindOrAdd(InTypefaceIndex);
+	TypefaceFontData.SetTypefaceName(Typeface.Name);
+	TypefaceFontData.SetTypefaceFontDataHash(HashCombine(0, GetTypeHash(Typeface.Font)));
+	TypefaceFontData.SetTypeFaceData(FaceData);
 }
 
 uint32 FCachedFontData::GetTypefaceFontDataHash(int32 InTypefaceEntryIndex) const
@@ -265,6 +254,16 @@ uint32 FCachedFontData::GetTypefaceFontDataHash(int32 InTypefaceEntryIndex) cons
 	return 0;
 }
 
+const FTypefaceFontData* FCachedFontData::GetTypefaceFontData(int32 InTypefaceEntryIndex) const
+{
+	return TypefaceFontDataMap.Find(InTypefaceEntryIndex);
+}
+
+FTypefaceFontData* FCachedFontData::GetTypefaceFontData(int32 InTypefaceEntryIndex)
+{
+	return TypefaceFontDataMap.Find(InTypefaceEntryIndex);
+}
+
 TSharedPtr<int32> FCachedFontData::GetCacheCounter(int32 InTypefaceEntryIndex)
 {
 	if (TypefaceFontDataMap.Contains(InTypefaceEntryIndex))
@@ -339,9 +338,9 @@ FText3DCachedMesh* FCachedFontData::GetGlyphMesh(uint32 InGlyphIndex, const FGly
 
 FT_Face FCachedFontData::GetFreeTypeFace(uint32 InTypefaceIndex)
 {
-	if (TypefaceFontDataMap.Contains(InTypefaceIndex))
+	if (const FTypefaceFontData* TypeFaceFontData = TypefaceFontDataMap.Find(InTypefaceIndex))
 	{
-		return TypefaceFontDataMap[InTypefaceIndex].GetTypeface();
+		return TypeFaceFontData->GetTypeFace();
 	}
 
 	return nullptr;
@@ -421,8 +420,6 @@ bool FCachedFontData::CleanupTypeface(uint32 InTypefaceIndex)
 
 		if (TypefaceFontData.GetCacheCount() <= 1)
 		{
-			TypefaceFontData.Reset();
-
 			bRemoved = true;
 		}
 	}
diff --git a/Engine/Plugins/Experimental/Text3D/Source/Text3D/Private/Subsystems/Text3DEngineSubsystem.h b/Engine/Plugins/Experimental/Text3D/Source/Text3D/Private/Subsystems/Text3DEngineSubsystem.h
--- a/Engine/Plugins/Experimental/Text3D/Source/Text3D/Private/Subsystems/Text3DEngineSubsystem.h
+++ b/Engine/Plugins/Experimental/Text3D/Source/Text3D/Private/Subsystems/Text3DEngineSubsystem.h
@@ -15,6 +15,7 @@ class UFont;
 class UMaterial;
 class UStaticMesh;
 class UWorld;
+struct FFontFaceData;
 
 USTRUCT()
 struct FGlyphMeshParameters
@@ -99,27 +100,14 @@ struct FTypefaceFontData
 
 	FTypefaceFontData();
 
-	void Reset()
-	{
-		FT_Done_Face(FreeTypeFace);
-		FreeTypeFace = nullptr;
-		TypefaceData.Reset();
-	}
-
 	int32 GetCacheCount() const
 	{
 		const int32 Count = CacheCounter.GetSharedReferenceCount();
 		return Count;
 	}
 
-	FT_Face& GetTypeface() { return FreeTypeFace; }
-
-	TArray<uint8>& GetTypefaceData() { return TypefaceData; }
-
-	void SetTypefaceData(const TArray<uint8>& InTypefaceData)
-	{
-		TypefaceData = InTypefaceData;
-	}
+	FT_Face GetTypeFace() const { return TypeFace; }
+	void SetTypeFace(FT_Face InTypeFace);
 
 	FName GetTypefaceName() const { return TypefaceName; }
 	void SetTypefaceName(FName InTypefaceName)
@@ -128,12 +116,14 @@ struct FTypefaceFontData
 	}
 
 	uint32 GetTypefaceFontDataHash() const { return TypefaceFontDataHash; }
-
 	void SetTypefaceFontDataHash(uint32 InDataHash)
 	{
 		TypefaceFontDataHash = InDataHash;
 	}
 
+	TSharedPtr<const FFontFaceData> GetTypeFaceData() const;
+	void SetTypeFaceData(const TSharedPtr<const FFontFaceData>& InTypeFaceData);
+
 	TMap<uint32, FCachedFontMeshes>& GetMeshes() { return Meshes; }
 
 	TSharedPtr<int32> GetCacheCounter() { return CacheCounter; }
@@ -147,11 +137,11 @@ private:
 	UPROPERTY()
 	TMap<uint32, FCachedFontMeshes> Meshes;
 
-	FName TypefaceName;
-	FT_Face FreeTypeFace;
-	TArray<uint8> TypefaceData;
+	FName TypefaceName = NAME_None;
 	TSharedPtr<int32> CacheCounter;
-	uint32 TypefaceFontDataHash;
+	uint32 TypefaceFontDataHash = 0;
+	FT_Face TypeFace = nullptr;
+	TWeakPtr<const FFontFaceData> TypeFaceDataWeak;
 };
 
 USTRUCT()
@@ -167,7 +157,6 @@ struct FCachedFontData
 	FString GetFontName() const;
 
 	UFont* GetFont() { return Font; }
-
 	void SetFont(UFont* InFont)
 	{
 		Font = InFont;
@@ -180,6 +169,8 @@ struct FCachedFontData
 	bool CleanupTypeface(uint32 InTypefaceIndex);
 
 	uint32 GetTypefaceFontDataHash(int32 InTypefaceEntryIndex) const;
+	const FTypefaceFontData* GetTypefaceFontData(int32 InTypefaceEntryIndex) const;
+	FTypefaceFontData* GetTypefaceFontData(int32 InTypefaceEntryIndex);
 	TSharedPtr<int32> GetCacheCounter(int32 InTypefaceEntryIndex);
 	TSharedPtr<int32> GetMeshesCacheCounter(const FGlyphMeshParameters& InParameters);
 
diff --git a/Engine/Plugins/Experimental/Text3D/Source/Text3D/Private/Text3DComponent.cpp b/Engine/Plugins/Experimental/Text3D/Source/Text3D/Private/Text3DComponent.cpp
--- a/Engine/Plugins/Experimental/Text3D/Source/Text3D/Private/Text3DComponent.cpp
+++ b/Engine/Plugins/Experimental/Text3D/Source/Text3D/Private/Text3DComponent.cpp
@@ -994,7 +994,7 @@ void UText3DComponent::OnTextRendererClassChanged()
 		TextRenderer = nullptr;
 	}
 
-	TextRenderer = NewObject<UText3DRendererBase>(this, TextRendererClass.Get(), RendererCDO->GetName());
+	TextRenderer = NewObject<UText3DRendererBase>(this, TextRendererClass.Get(), RendererCDO->GetName(), RF_Transient);
 
 	const AActor* ContextActor = GetOwner();
 	UE_LOG(LogText3D, Log, TEXT("%s : Text3D renderer changed : old %s - new %s"), !!ContextActor ? *ContextActor->GetActorNameOrLabel() : TEXT("Invalid owner"), *OldRendererName.ToString(), *TextRenderer->GetName().ToString());
diff --git a/Engine/Plugins/Experimental/Text3D/Source/Text3D/Private/Utilities/Text3DUtilities.cpp b/Engine/Plugins/Experimental/Text3D/Source/Text3D/Private/Utilities/Text3DUtilities.cpp
--- a/Engine/Plugins/Experimental/Text3D/Source/Text3D/Private/Utilities/Text3DUtilities.cpp
+++ b/Engine/Plugins/Experimental/Text3D/Source/Text3D/Private/Utilities/Text3DUtilities.cpp
@@ -208,97 +208,3 @@ bool UE::Text3D::Utilities::GetFontFaces(const UFont* InFont, TArray<UFontFace*>
 
 	return !OutFontFaces.IsEmpty();
 }
-
-TArray<FText3DCharacterMetadata> UE::Text3D::Utilities::GetCharacterMetadata(const FTypefaceEntry& InTypeFace, const FString& InText)
-{
-	const FFontFaceDataConstPtr FaceData = InTypeFace.Font.GetFontFaceData();
-
-	TArray<FText3DCharacterMetadata> CharacterIndexes;
-	CharacterIndexes.Reserve(InText.Len() + 1);
-
-	if (FaceData.IsValid() && FaceData->HasData() && !FaceData->GetData().IsEmpty())
-	{
-		const TArray<uint8> Data = FaceData->GetData();
-
-		FT_Face FreeTypeFace;
-		FT_New_Memory_Face(FText3DModule::GetFreeTypeLibrary(), Data.GetData(), Data.Num(), 0, &FreeTypeFace);
-
-		if (FreeTypeFace)
-		{
-			int32 TextIndex = 0;
-			int32 VisibleIndex = 0;
-			int32 LayoutIndex = 0;
-
-			FT_Set_Pixel_Sizes(FreeTypeFace, 0, UE::Text3D::Metrics::FontSize);
-
-			for (const TCHAR Character : InText)
-			{
-				if (!FChar::IsControl(Character) && FChar::IsPrint(Character))
-				{
-					if (!FChar::IsWhitespace(Character))
-					{
-						// Check if the font can render this character
-						FT_UInt GlyphIndex = FT_Get_Char_Index(FreeTypeFace, Character);
-        
-						if (GlyphIndex > 0)
-						{
-							// Load the glyph to check if it has an actual outline or bitmap
-							const FT_Error Code = FT_Load_Glyph(FreeTypeFace, GlyphIndex, FT_LOAD_DEFAULT);
-							if (Code == 0)
-							{
-								// Check if glyph has bitmap data or outline data
-								if ((FreeTypeFace->glyph->format == FT_GLYPH_FORMAT_BITMAP
-									&& FreeTypeFace->glyph->bitmap.width > 0 && FreeTypeFace->glyph->bitmap.rows > 0)
-									|| (FreeTypeFace->glyph->format == FT_GLYPH_FORMAT_OUTLINE && FreeTypeFace->glyph->outline.n_points > 0))
-								{
-									FText3DCharacterMetadata Metadata;
-									Metadata.GlyphIndex = GlyphIndex;
-									Metadata.LayoutIndex = LayoutIndex;
-									Metadata.VisibleIndex = VisibleIndex;
-									Metadata.TextIndex = TextIndex;
-									Metadata.Height = FreeTypeFace->glyph->metrics.height / UE::Text3D::Metrics::FontSize;
-									Metadata.Width = FreeTypeFace->glyph->metrics.width / UE::Text3D::Metrics::FontSize;
-
-									CharacterIndexes.Emplace(Metadata);
-
-									VisibleIndex++;
-								}
-							}
-						}
-					}
-
-					LayoutIndex++;
-				}
-
-				TextIndex++;
-			}
-	
-			FT_Done_Face(FreeTypeFace);
-		}
-	}
-
-	return CharacterIndexes;
-}
-
-FText3DTypeFaceMetrics UE::Text3D::Utilities::GetTypeFaceMetrics(const FTypefaceEntry& InTypeFace)
-{
-	FText3DTypeFaceMetrics Metrics;
-
-	const FFontFaceDataConstPtr FaceData = InTypeFace.Font.GetFontFaceData();
-
-	if (FaceData.IsValid() && FaceData->HasData() && !FaceData->GetData().IsEmpty())
-	{
-		const TArray<uint8> Data = FaceData->GetData();
-
-		FT_Face FreeTypeFace;
-		FT_New_Memory_Face(FText3DModule::GetFreeTypeLibrary(), Data.GetData(), Data.Num(), 0, &FreeTypeFace);
-
-		Metrics.FontAscender = FreeTypeFace->ascender * UE::Text3D::Metrics::FontSizeInverse;
-		Metrics.FontDescender = FreeTypeFace->descender * UE::Text3D::Metrics::FontSizeInverse;
-		Metrics.FontHeight = FreeTypeFace->height * UE::Text3D::Metrics::FontSizeInverse;
-
-		FT_Done_Face(FreeTypeFace);
-	}
-
-	return Metrics;
-}
diff --git a/Engine/Plugins/Experimental/Text3D/Source/Text3D/Public/Characters/Text3DCharacterBase.h b/Engine/Plugins/Experimental/Text3D/Source/Text3D/Public/Characters/Text3DCharacterBase.h
--- a/Engine/Plugins/Experimental/Text3D/Source/Text3D/Public/Characters/Text3DCharacterBase.h
+++ b/Engine/Plugins/Experimental/Text3D/Source/Text3D/Public/Characters/Text3DCharacterBase.h
@@ -31,16 +31,15 @@ public:
 		return Character;
 	}
 #endif
-	
-	void SetCharacterMetadata(const FText3DCharacterMetadata& InMetadata);
-	const FText3DCharacterMetadata& GetCharacterMetadata() const
-	{
-		return Metadata;
-	}
-	FText3DCharacterMetadata& GetCharacterMetadata()
-	{
-		return Metadata;
-	}
+
+	void SetGlyphIndex(uint32 InGlyphIndex);
+	uint32 GetGlyphIndex() const;
+
+	void SetMeshBounds(const FBox& InBounds);
+	const FBox& GetMeshBounds() const;
+
+	void SetMeshOffset(const FVector& InOffset);
+	const FVector& GetMeshOffset() const;
 
 	UFUNCTION()
 	TEXT3D_API void SetRelativeLocation(const FVector& InLocation);
@@ -110,6 +109,12 @@ protected:
 	/** Final transform after all extensions are applied */
 	FTransform Transform = FTransform::Identity;
 
-	/** Character metadata */
-	FText3DCharacterMetadata Metadata;
+	/** Glyph that represents this character */
+	uint32 GlyphIndex = 0;
+
+	/** Actual mesh render bounds */
+	FBox MeshBounds = FBox(ForceInitToZero);
+
+	/** Offset around mesh due to font face */
+	FVector MeshOffset = FVector::ZeroVector;
 };
diff --git a/Engine/Plugins/Experimental/Text3D/Source/Text3D/Public/Extensions/Text3DCharacterExtensionBase.h b/Engine/Plugins/Experimental/Text3D/Source/Text3D/Public/Extensions/Text3DCharacterExtensionBase.h
--- a/Engine/Plugins/Experimental/Text3D/Source/Text3D/Public/Extensions/Text3DCharacterExtensionBase.h
+++ b/Engine/Plugins/Experimental/Text3D/Source/Text3D/Public/Extensions/Text3DCharacterExtensionBase.h
@@ -37,6 +37,9 @@ public:
 		return {};
 	}
 
+	/** Allocate visible characters to store data */
+	virtual void AllocateCharacters(uint16 InCount) {}
+
 #if WITH_EDITORONLY_DATA
 	/** Used for detail customization within editor */
 	UPROPERTY(EditInstanceOnly, Transient, Category="Character", meta=(TextCharacterSelector, AllowPrivateAccess=true))
diff --git a/Engine/Plugins/Experimental/Text3D/Source/Text3D/Public/Extensions/Text3DDefaultCharacterExtension.h b/Engine/Plugins/Experimental/Text3D/Source/Text3D/Public/Extensions/Text3DDefaultCharacterExtension.h
--- a/Engine/Plugins/Experimental/Text3D/Source/Text3D/Public/Extensions/Text3DDefaultCharacterExtension.h
+++ b/Engine/Plugins/Experimental/Text3D/Source/Text3D/Public/Extensions/Text3DDefaultCharacterExtension.h
@@ -17,13 +17,9 @@ protected:
 	virtual uint16 GetCharacterCount() const override;
 	virtual UText3DCharacterBase* GetCharacter(uint16 InCharacter) const override;
 	virtual TConstArrayView<UText3DCharacterBase*> GetCharacters() const override;
+	virtual void AllocateCharacters(uint16 InCount) override;
 	//~ End UText3DGeometryExtensionBase
 
-	//~ Begin UText3DExtensionBase
-	virtual EText3DExtensionResult PreRendererUpdate(EText3DRendererFlags InFlag) override;
-	virtual EText3DExtensionResult PostRendererUpdate(EText3DRendererFlags InFlag) override;
-	//~ End UText3DExtensionBase
-
 	/** Allocate character data by using pool or creating new object */
 	void AllocateTextCharacters(uint16 InCharacterCount);
 
diff --git a/Engine/Plugins/Experimental/Text3D/Source/Text3D/Public/Extensions/Text3DDefaultGeometryExtension.h b/Engine/Plugins/Experimental/Text3D/Source/Text3D/Public/Extensions/Text3DDefaultGeometryExtension.h
--- a/Engine/Plugins/Experimental/Text3D/Source/Text3D/Public/Extensions/Text3DDefaultGeometryExtension.h
+++ b/Engine/Plugins/Experimental/Text3D/Source/Text3D/Public/Extensions/Text3DDefaultGeometryExtension.h
@@ -102,9 +102,10 @@ protected:
 	//~ End UObject
 
 	//~ Begin UText3DGeometryExtensionBase
-	virtual UStaticMesh* GetGlyphMesh(UText3DCharacterBase* InCharacter) const override;
+	virtual UStaticMesh* FindOrLoadGlyphMesh(UText3DCharacterBase* InCharacter) const override;
 	virtual EText3DHorizontalTextAlignment GetGlyphHAlignment() const override;
 	virtual EText3DVerticalTextAlignment GetGlyphVAlignment() const override;
+	virtual FTypefaceFontData* GetTypefaceFontData() const override;
 	//~ End UText3DGeometryExtensionBase
 
 	//~ Begin UText3DExtensionBase
diff --git a/Engine/Plugins/Experimental/Text3D/Source/Text3D/Public/Extensions/Text3DGeometryExtensionBase.h b/Engine/Plugins/Experimental/Text3D/Source/Text3D/Public/Extensions/Text3DGeometryExtensionBase.h
--- a/Engine/Plugins/Experimental/Text3D/Source/Text3D/Public/Extensions/Text3DGeometryExtensionBase.h
+++ b/Engine/Plugins/Experimental/Text3D/Source/Text3D/Public/Extensions/Text3DGeometryExtensionBase.h
@@ -9,6 +9,7 @@
 class FFreeTypeFace;
 class UStaticMesh;
 class UText3DCharacterBase;
+struct FTypefaceFontData;
 
 /** Extension that handles geometry data for Text3D */
 UCLASS(MinimalAPI, Abstract)
@@ -21,7 +22,7 @@ public:
 		: UText3DExtensionBase(/** Priority */1)
 	{}
 
-	virtual UStaticMesh* GetGlyphMesh(UText3DCharacterBase* InCharacter) const
+	virtual UStaticMesh* FindOrLoadGlyphMesh(UText3DCharacterBase* InCharacter) const
 	{
 		return nullptr;
 	}
@@ -35,4 +36,9 @@ public:
 	{
 		return EText3DVerticalTextAlignment::Bottom;
 	}
+
+	virtual FTypefaceFontData* GetTypefaceFontData() const
+	{
+		return nullptr;
+	}
 };
diff --git a/Engine/Plugins/Experimental/Text3D/Source/Text3D/Public/Text3DComponent.h b/Engine/Plugins/Experimental/Text3D/Source/Text3D/Public/Text3DComponent.h
--- a/Engine/Plugins/Experimental/Text3D/Source/Text3D/Public/Text3DComponent.h
+++ b/Engine/Plugins/Experimental/Text3D/Source/Text3D/Public/Text3DComponent.h
@@ -20,7 +20,6 @@ class UText3DMaterialExtensionBase;
 class UText3DRendererBase;
 class UText3DRenderingExtensionBase;
 class UText3DLayoutEffectBase;
-struct FText3DShapedGlyphText;
 struct FTypefaceEntry;
 
 UCLASS(MinimalAPI
diff --git a/Engine/Plugins/Experimental/Text3D/Source/Text3D/Public/Text3DTypes.h b/Engine/Plugins/Experimental/Text3D/Source/Text3D/Public/Text3DTypes.h
--- a/Engine/Plugins/Experimental/Text3D/Source/Text3D/Public/Text3DTypes.h
+++ b/Engine/Plugins/Experimental/Text3D/Source/Text3D/Public/Text3DTypes.h
@@ -212,26 +212,6 @@ struct FText3DMaterialGroupKey
 	EText3DMaterialStyle Style = EText3DMaterialStyle::Invalid;
 };
 
-struct FText3DCharacterMetadata
-{
-	/** Index of character based on text string (control code, whitespace, ...) */
-	uint16 TextIndex = 0;
-	/** Index of character based on visibility (only renderable character) */
-	uint16 VisibleIndex = 0;
-	/** Index of character based on layout index (include whitespace) */
-	uint16 LayoutIndex = 0;
-	/** Index of character based on glyph index (typeface font specific) */
-	uint16 GlyphIndex = 0;
-	/** Width of the character based on typeface and font metrics */
-	uint16 Width = 0;
-	/** Height of the character based on typeface and font metrics */
-	uint16 Height = 0;
-	/** Actual render bounds */
-	FBox MeshBounds = FBox(ForceInitToZero);
-	/** Offset around glyph due to font face */
-	FVector MeshOffset = FVector::ZeroVector;
-};
-
 struct FText3DTypeFaceMetrics
 {
 	float FontHeight = 0.f;
diff --git a/Engine/Plugins/Experimental/Text3D/Source/Text3D/Public/Utilities/Text3DUtilities.h b/Engine/Plugins/Experimental/Text3D/Source/Text3D/Public/Utilities/Text3DUtilities.h
--- a/Engine/Plugins/Experimental/Text3D/Source/Text3D/Public/Utilities/Text3DUtilities.h
+++ b/Engine/Plugins/Experimental/Text3D/Source/Text3D/Public/Utilities/Text3DUtilities.h
@@ -30,10 +30,4 @@ namespace UE::Text3D::Utilities
 	
 	/** Get font faces that compose the font */
 	TEXT3D_API bool GetFontFaces(const UFont* InFont, TArray<UFontFace*>& OutFontFaces);
-
-	/** Get indexes of character for a text based on typeface */
-	TEXT3D_API TArray<FText3DCharacterMetadata> GetCharacterMetadata(const FTypefaceEntry& InTypeFace, const FString& InText);
-
-	/** Get type face font metrics */
-	TEXT3D_API FText3DTypeFaceMetrics GetTypeFaceMetrics(const FTypefaceEntry& InTypeFace);
 }
\ No newline at end of file
