MotionDesign : Text3D
- Fixed UnrealServer compilation warnings caused by FreeType and HarfBuzz macro redefinition

diff --git a/Engine/Plugins/Experimental/Text3D/Source/Text3D/Private/Extensions/Text3DDefaultGeometryExtension.cpp b/Engine/Plugins/Experimental/Text3D/Source/Text3D/Private/Extensions/Text3DDefaultGeometryExtension.cpp
--- a/Engine/Plugins/Experimental/Text3D/Source/Text3D/Private/Extensions/Text3DDefaultGeometryExtension.cpp
+++ b/Engine/Plugins/Experimental/Text3D/Source/Text3D/Private/Extensions/Text3DDefaultGeometryExtension.cpp
@@ -139,11 +139,14 @@ void UText3DDefaultGeometryExtension::PostEditChangeProperty(FPropertyChangedEve
 
 UStaticMesh* UText3DDefaultGeometryExtension::FindOrLoadGlyphMesh(UText3DCharacterBase* InCharacter) const
 {
-	const FText3DCachedMesh* GlyphMesh = CachedFontData->GetGlyphMesh(InCharacter->GetGlyphIndex(), GlyphMeshParameters);
-	checkf(GlyphMesh != nullptr, TEXT("Invalid glyph mesh returned"));
-	InCharacter->SetMeshBounds(GlyphMesh->MeshBounds);
-	InCharacter->SetMeshOffset(GlyphMesh->MeshOffset);
-	return GlyphMesh->Mesh;
+	if (const FText3DCachedMesh* GlyphMesh = CachedFontData->GetGlyphMesh(InCharacter->GetGlyphIndex(), GlyphMeshParameters))
+	{
+		InCharacter->SetMeshBounds(GlyphMesh->MeshBounds);
+		InCharacter->SetMeshOffset(GlyphMesh->MeshOffset);
+		return GlyphMesh->Mesh;
+	}
+
+	return nullptr;
 }
 
 EText3DHorizontalTextAlignment UText3DDefaultGeometryExtension::GetGlyphHAlignment() const
diff --git a/Engine/Plugins/Experimental/Text3D/Source/Text3D/Private/Extensions/Text3DDefaultLayoutExtension.cpp b/Engine/Plugins/Experimental/Text3D/Source/Text3D/Private/Extensions/Text3DDefaultLayoutExtension.cpp
--- a/Engine/Plugins/Experimental/Text3D/Source/Text3D/Private/Extensions/Text3DDefaultLayoutExtension.cpp
+++ b/Engine/Plugins/Experimental/Text3D/Source/Text3D/Private/Extensions/Text3DDefaultLayoutExtension.cpp
@@ -8,17 +8,21 @@
 #include "Engine/Font.h"
 #include "Extensions/Text3DCharacterExtensionBase.h"
 #include "Extensions/Text3DGeometryExtensionBase.h"
-#include "Fonts/SlateTextShaper.h"
 #include "Framework/Text/PlainTextLayoutMarshaller.h"
 #include "LayoutBuilders/Text3DLayout.h"
 #include "LayoutBuilders/Text3DLayoutShaper.h"
 #include "LayoutBuilders/Text3DShapedGlyphText.h"
+#include "Logs/Text3DLogs.h"
 #include "Misc/EnumerateRange.h"
 #include "Styling/StyleDefaults.h"
 #include "Styling/SlateTypes.h"
 #include "Subsystems/Text3DEngineSubsystem.h"
 #include "Text3DComponent.h"
 
+#if WITH_FREETYPE
+#include "Fonts/SlateTextShaper.h"
+#endif
+
 FName UText3DDefaultLayoutExtension::GetUseMaxWidthPropertyName()
 {
 	return GET_MEMBER_NAME_CHECKED(UText3DDefaultLayoutExtension, bUseMaxWidth);
@@ -235,7 +239,10 @@ EText3DExtensionResult UText3DDefaultLayoutExtension::PreRendererUpdate(EText3DR
 	ShapedText->GlyphIndexToFontFace.Reserve(CharacterCount);
 	ShapedText->Kernings.SetNum(CharacterCount);
 
+	UText3DGeometryExtensionBase* GeometryExtension = Text3DComponent->GetGeometryExtension();
 	uint16 CharacterIndex = 0;
+
+#if WITH_FREETYPE
 	FT_Face FontFace = nullptr;
 	for (TConstEnumerateRef<FShapedGlyphLine> ShapedLine : EnumerateRange(ShapedText->Lines))
 	{
@@ -270,7 +277,6 @@ EText3DExtensionResult UText3DDefaultLayoutExtension::PreRendererUpdate(EText3DR
 		}
 	}
 
-	UText3DGeometryExtensionBase* GeometryExtension = Text3DComponent->GetGeometryExtension();
 	if (EnumHasAnyFlags(Text3DComponent->GetUpdateFlags(), EText3DRendererFlags::Geometry))
 	{
 		FTypefaceFontData* TypefaceFontData = GeometryExtension->GetTypefaceFontData();
@@ -288,6 +294,10 @@ EText3DExtensionResult UText3DDefaultLayoutExtension::PreRendererUpdate(EText3DR
 			TypeFaceMetrics.FontDescender = FontFace->size->metrics.descender * UE::Text3D::Metrics::FontSizeInverse;
 		}
 	}
+#else
+	UE_LOG(LogText3D, Warning, TEXT("FreeType is not available, cannot proceed without it"))
+	return EText3DExtensionResult::Failed;
+#endif
 
 	ShapedText->Tracking = Tracking;
 	ShapedText->WordSpacing = WordSpacing;
diff --git a/Engine/Plugins/Experimental/Text3D/Source/Text3D/Private/GeometryBuilders/Text3DGlyphLoader.cpp b/Engine/Plugins/Experimental/Text3D/Source/Text3D/Private/GeometryBuilders/Text3DGlyphLoader.cpp
--- a/Engine/Plugins/Experimental/Text3D/Source/Text3D/Private/GeometryBuilders/Text3DGlyphLoader.cpp
+++ b/Engine/Plugins/Experimental/Text3D/Source/Text3D/Private/GeometryBuilders/Text3DGlyphLoader.cpp
@@ -10,6 +10,7 @@
 
 using namespace UE::Geometry;
 
+#if WITH_FREETYPE
 FText3DGlyphLoader::FText3DGlyphLoader(const FT_GlyphSlot InGlyph)
 	: Root(MakeShared<FText3DGlyphContourNode>(nullptr, false, true))
 	, EndIndex(-1)
@@ -37,6 +38,7 @@
 
 	FixParityAndDivide(Root, false);
 }
+#endif
 
 TText3DGlyphContourNodeShared FText3DGlyphLoader::GetContourList() const
 {
@@ -288,12 +290,146 @@
 	return Result.GetSafeNormal();
 }
 
+#if WITH_FREETYPE
 TSharedPtr<FPolygon2f> FText3DGlyphLoader::CreateContour(const FT_Outline Outline, const int32 ContourIndex)
 {
 	const TSharedPtr<FPolygon2f> Contour = ProcessFreetypeOutline(Outline, ContourIndex);
 	return Contour.IsValid() ? RemoveBadPoints(Contour) : Contour;
 }
 
+TSharedPtr<FPolygon2f> FText3DGlyphLoader::ProcessFreetypeOutline(const FT_Outline Outline, const int32 ContourIndex)
+{
+	TSharedPtr<FPolygon2f> Contour;
+
+	const int32 StartIndex = EndIndex + 1;
+	EndIndex = Outline.contours[ContourIndex];
+	const int32 ContourLength = EndIndex - StartIndex + 1;
+
+	if (ContourLength < 3)
+	{
+		return Contour;
+	}
+
+	ProcessedContour.Empty();
+
+	const FT_Vector* const Points = Outline.points + StartIndex;
+	auto ToFVector2D = [Points](const int32 Index)
+	{
+		const FT_Vector Point = Points[Index];
+		return FVector2D(Point.x, Point.y);
+	};
+
+	FVector2D Prev;
+	FVector2D Current = ToFVector2D(ContourLength - 1);
+	FVector2D Next = ToFVector2D(0);
+	FVector2D NextNext = ToFVector2D(1);
+
+
+	const char* const Tags = Outline.tags + StartIndex;
+	auto Tag = [Tags](int32 Index)
+	{
+		return FT_CURVE_TAG(Tags[Index]);
+	};
+
+	int32 TagPrev = 0;
+	int32 TagCurrent = Tag(ContourLength - 1);
+	int32 TagNext = Tag(0);
+
+
+	FVector2D& FirstPositionLocal = FirstPosition;
+	TDoubleLinkedList<FVector2f>& ProcessedContourLocal = ProcessedContour;
+	auto ContourIsBad = [&FirstPositionLocal, &ProcessedContourLocal](const FVector2D Point)
+	{
+		if (ProcessedContourLocal.Num() == 0)
+		{
+			FirstPositionLocal = Point;
+			return false;
+		}
+
+		return (Point - FVector2D(ProcessedContourLocal.GetHead()->GetValue())).IsNearlyZero();
+	};
+
+	for (int32 Index = 0; Index < ContourLength; Index++)
+	{
+		const int32 NextIndex = (Index + 1) % ContourLength;
+
+		Prev = Current;
+		Current = Next;
+		Next = NextNext;
+		NextNext = ToFVector2D((NextIndex + 1) % ContourLength);
+
+		TagPrev = TagCurrent;
+		TagCurrent = TagNext;
+		TagNext = Tag(NextIndex);
+
+		if (TagCurrent == FT_Curve_Tag_On)
+		{
+			if (TagNext == FT_Curve_Tag_Cubic || TagNext == FT_Curve_Tag_Conic)
+			{
+				continue;
+			}
+
+			if (ContourIsBad(Current))
+			{
+				return Contour;
+			}
+
+			if (TagNext == FT_Curve_Tag_On && (Current - Next).IsNearlyZero())
+			{
+				continue;
+			}
+
+			FLine(Current).Add(this);
+		}
+		else if (TagCurrent == FT_Curve_Tag_Conic)
+		{
+			FVector2D A;
+
+			if (TagPrev == FT_Curve_Tag_On)
+			{
+				if (ContourIsBad(Prev))
+				{
+					return Contour;
+				}
+
+				A = Prev;
+			}
+			else
+			{
+				A = (Prev + Current) / 2.f;
+			}
+
+			FQuadraticCurve(A, Current, TagNext == FT_Curve_Tag_Conic ? (Current + Next) / 2.f : Next).Add(this);
+		}
+		else if (TagCurrent == FT_Curve_Tag_Cubic && TagNext == FT_Curve_Tag_Cubic)
+		{
+			if (ContourIsBad(Prev))
+			{
+				return Contour;
+			}
+
+			FCubicCurve(Prev, Current, Next, NextNext).Add(this);
+		}
+	}
+
+	if (ProcessedContour.Num() < 3)
+	{
+		return Contour;
+	}
+
+
+	Contour = MakeShared<FPolygon2f>();
+
+	for (const TDoubleLinkedList<FVector2f>::TDoubleLinkedListNode* Point = ProcessedContour.GetTail(); Point; Point =
+		 Point->GetPrevNode())
+	{
+		Contour->AppendVertex(Point->GetValue());
+	}
+
+	return Contour;
+}
+#endif
+
 bool FText3DGlyphLoader::IsInside(const TSharedPtr<const FPolygon2f> ContourA,
 							const TSharedPtr<const FPolygon2f> ContourB) const
 {
@@ -434,138 +570,6 @@
 	}
 }
 
-TSharedPtr<FPolygon2f> FText3DGlyphLoader::ProcessFreetypeOutline(const FT_Outline Outline, const int32 ContourIndex)
-{
-	TSharedPtr<FPolygon2f> Contour;
-
-	const int32 StartIndex = EndIndex + 1;
-	EndIndex = Outline.contours[ContourIndex];
-	const int32 ContourLength = EndIndex - StartIndex + 1;
-
-	if (ContourLength < 3)
-	{
-		return Contour;
-	}
-
-	ProcessedContour.Empty();
-
-	const FT_Vector* const Points = Outline.points + StartIndex;
-	auto ToFVector2D = [Points](const int32 Index)
-	{
-		const FT_Vector Point = Points[Index];
-		return FVector2D(Point.x, Point.y);
-	};
-
-	FVector2D Prev;
-	FVector2D Current = ToFVector2D(ContourLength - 1);
-	FVector2D Next = ToFVector2D(0);
-	FVector2D NextNext = ToFVector2D(1);
-
-
-	const char* const Tags = Outline.tags + StartIndex;
-	auto Tag = [Tags](int32 Index)
-	{
-		return FT_CURVE_TAG(Tags[Index]);
-	};
-
-	int32 TagPrev = 0;
-	int32 TagCurrent = Tag(ContourLength - 1);
-	int32 TagNext = Tag(0);
-
-
-	FVector2D& FirstPositionLocal = FirstPosition;
-	TDoubleLinkedList<FVector2f>& ProcessedContourLocal = ProcessedContour;
-	auto ContourIsBad = [&FirstPositionLocal, &ProcessedContourLocal](const FVector2D Point)
-	{
-		if (ProcessedContourLocal.Num() == 0)
-		{
-			FirstPositionLocal = Point;
-			return false;
-		}
-
-		return (Point - FVector2D(ProcessedContourLocal.GetHead()->GetValue())).IsNearlyZero();
-	};
-
-	for (int32 Index = 0; Index < ContourLength; Index++)
-	{
-		const int32 NextIndex = (Index + 1) % ContourLength;
-
-		Prev = Current;
-		Current = Next;
-		Next = NextNext;
-		NextNext = ToFVector2D((NextIndex + 1) % ContourLength);
-
-		TagPrev = TagCurrent;
-		TagCurrent = TagNext;
-		TagNext = Tag(NextIndex);
-
-		if (TagCurrent == FT_Curve_Tag_On)
-		{
-			if (TagNext == FT_Curve_Tag_Cubic || TagNext == FT_Curve_Tag_Conic)
-			{
-				continue;
-			}
-
-			if (ContourIsBad(Current))
-			{
-				return Contour;
-			}
-
-			if (TagNext == FT_Curve_Tag_On && (Current - Next).IsNearlyZero())
-			{
-				continue;
-			}
-
-			FLine(Current).Add(this);
-		}
-		else if (TagCurrent == FT_Curve_Tag_Conic)
-		{
-			FVector2D A;
-
-			if (TagPrev == FT_Curve_Tag_On)
-			{
-				if (ContourIsBad(Prev))
-				{
-					return Contour;
-				}
-
-				A = Prev;
-			}
-			else
-			{
-				A = (Prev + Current) / 2.f;
-			}
-
-			FQuadraticCurve(A, Current, TagNext == FT_Curve_Tag_Conic ? (Current + Next) / 2.f : Next).Add(this);
-		}
-		else if (TagCurrent == FT_Curve_Tag_Cubic && TagNext == FT_Curve_Tag_Cubic)
-		{
-			if (ContourIsBad(Prev))
-			{
-				return Contour;
-			}
-
-			FCubicCurve(Prev, Current, Next, NextNext).Add(this);
-		}
-	}
-
-	if (ProcessedContour.Num() < 3)
-	{
-		return Contour;
-	}
-
-
-	Contour = MakeShared<FPolygon2f>();
-
-	for (const TDoubleLinkedList<FVector2f>::TDoubleLinkedListNode* Point = ProcessedContour.GetTail(); Point; Point =
-		 Point->GetPrevNode())
-	{
-		Contour->AppendVertex(Point->GetValue());
-	}
-
-	return Contour;
-}
-
 bool FText3DGlyphLoader::IsBadNormal(const TSharedPtr<FPolygon2f> Contour, const int32 Point) const
 {
 	FVector2f ToPrev;
diff --git a/Engine/Plugins/Experimental/Text3D/Source/Text3D/Private/GeometryBuilders/Text3DGlyphLoader.h b/Engine/Plugins/Experimental/Text3D/Source/Text3D/Private/GeometryBuilders/Text3DGlyphLoader.h
--- a/Engine/Plugins/Experimental/Text3D/Source/Text3D/Private/GeometryBuilders/Text3DGlyphLoader.h
+++ b/Engine/Plugins/Experimental/Text3D/Source/Text3D/Private/GeometryBuilders/Text3DGlyphLoader.h
@@ -15,7 +15,9 @@ using UE::Geometry::FPolygon2f;
 class FText3DGlyphLoader final
 {
 public:
+#if WITH_FREETYPE
 	explicit FText3DGlyphLoader(const FT_GlyphSlot InGlyph);
+#endif
 
 	TText3DGlyphContourNodeShared GetContourList() const;
 
@@ -134,9 +136,10 @@ private:
 	// Path
 	TArray<int32> DividedContourIDs;
 
-
-
+#if WITH_FREETYPE
 	TSharedPtr<FPolygon2f> CreateContour(const FT_Outline Outline, const int32 ContourIndex);
+	TSharedPtr<FPolygon2f> ProcessFreetypeOutline(const FT_Outline Outline, const int32 ContourIndex);
+#endif
 
 	/**
 	 * Checks if ContourA is inside ContourB.
@@ -158,7 +161,6 @@ private:
 	 */
 	void FixParityAndDivide(const TText3DGlyphContourNodeShared Node, const bool bClockwiseIn);
 
-	TSharedPtr<FPolygon2f> ProcessFreetypeOutline(const FT_Outline Outline, const int32 ContourIndex);
 	TSharedPtr<FPolygon2f> RemoveBadPoints(const TSharedPtr<FPolygon2f> Contour) const;
 
 	bool IsBadNormal(const TSharedPtr<FPolygon2f> Contour, const int32 Point) const;
diff --git a/Engine/Plugins/Experimental/Text3D/Source/Text3D/Private/Subsystems/Text3DEngineSubsystem.cpp b/Engine/Plugins/Experimental/Text3D/Source/Text3D/Private/Subsystems/Text3DEngineSubsystem.cpp
--- a/Engine/Plugins/Experimental/Text3D/Source/Text3D/Private/Subsystems/Text3DEngineSubsystem.cpp
+++ b/Engine/Plugins/Experimental/Text3D/Source/Text3D/Private/Subsystems/Text3DEngineSubsystem.cpp
@@ -5,7 +5,6 @@
 #include "Engine/Font.h"
 #include "Engine/Engine.h"
 #include "Engine/StaticMesh.h"
-#include "Fonts/FontCacheFreeType.h"
 #include "GeometryBuilders/Text3DGlyphContourNode.h"
 #include "GeometryBuilders/Text3DGlyphLoader.h"
 #include "GeometryBuilders/Text3DGlyphMeshBuilder.h"
@@ -14,6 +13,10 @@
 #include "Misc/FileHelper.h"
 #include "UObject/ConstructorHelpers.h"
 
+#if WITH_FREETYPE
+#include "Fonts/FontCacheFreeType.h"
+#endif
+
 FCachedFontMeshes::FCachedFontMeshes()
 {
 	CacheCounter = MakeShared<int32>(0);
@@ -187,20 +190,22 @@ FTypefaceFontData::FTypefaceFontData()
 	CacheCounter = MakeShared<int32>(0);
 }
 
+#if WITH_FREETYPE
 void FTypefaceFontData::SetTypeFace(FT_Face InTypeFace)
 {
 	TypeFace = InTypeFace;
 }
 
-TSharedPtr<const FFontFaceData> FTypefaceFontData::GetTypeFaceData() const
+void FTypefaceFontData::SetTypeFaceData(const TSharedPtr<const FFontFaceData>& InTypeFaceData)
 {
-	return TypeFaceDataWeak.Pin();
+	TypeFaceDataWeak = InTypeFaceData;
 }
 
-void FTypefaceFontData::SetTypeFaceData(const TSharedPtr<const FFontFaceData>& InTypeFaceData)
+TSharedPtr<const FFontFaceData> FTypefaceFontData::GetTypeFaceData() const
 {
-	TypeFaceDataWeak = InTypeFaceData;
+	return TypeFaceDataWeak.Pin();
 }
+#endif
 
 FCachedFontData::FCachedFontData()
 {
@@ -241,7 +246,9 @@ void FCachedFontData::LoadFreeTypeFace(uint32 InTypefaceIndex)
 	FTypefaceFontData& TypefaceFontData = TypefaceFontDataMap.FindOrAdd(InTypefaceIndex);
 	TypefaceFontData.SetTypefaceName(Typeface.Name);
 	TypefaceFontData.SetTypefaceFontDataHash(HashCombine(0, GetTypeHash(Typeface.Font)));
+#if WITH_FREETYPE
 	TypefaceFontData.SetTypeFaceData(FaceData);
+#endif
 }
 
 uint32 FCachedFontData::GetTypefaceFontDataHash(int32 InTypefaceEntryIndex) const
@@ -311,7 +318,7 @@ FText3DCachedMesh* FCachedFontData::GetGlyphMesh(uint32 InGlyphIndex, const FGly
 	const FString StaticMeshName = FString::Printf(TEXT("Text3D_Char_%u_%u"), HashGroup, HashParameters);
 
 	const TText3DGlyphContourNodeShared Root = GetGlyphContours(InGlyphIndex, TypefaceIndex);
-	if (Root->Children.Num() == 0)
+	if (!Root || Root->Children.IsEmpty())
 	{
 		return nullptr;
 	}
@@ -336,6 +343,7 @@ FText3DCachedMesh* FCachedFontData::GetGlyphMesh(uint32 InGlyphIndex, const FGly
 	return &CachedMesh;
 }
 
+#if WITH_FREETYPE
 FT_Face FCachedFontData::GetFreeTypeFace(uint32 InTypefaceIndex)
 {
 	if (const FTypefaceFontData* TypeFaceFontData = TypefaceFontDataMap.Find(InTypefaceIndex))
@@ -345,6 +353,7 @@ FT_Face FCachedFontData::GetFreeTypeFace(uint32 InTypefaceIndex)
 
 	return nullptr;
 }
+#endif
 
 FString FCachedFontData::GetFontName() const
 {
@@ -358,6 +367,7 @@ FString FCachedFontData::GetFontName() const
 
 TText3DGlyphContourNodeShared FCachedFontData::GetGlyphContours(uint32 InGlyphIndex, int32 InTypefaceEntryIndex)
 {
+#if WITH_FREETYPE
 	FT_Face Face = GetFreeTypeFace(InTypefaceEntryIndex);
 	if (!Face)
 	{
@@ -373,6 +383,10 @@ TText3DGlyphContourNodeShared FCachedFontData::GetGlyphContours(uint32 InGlyphIn
 	TText3DGlyphContourNodeShared Root = GlyphLoader.GetContourList();
 
 	return Root;
+#else
+	UE_LOG(LogText3D, Warning, TEXT("FreeType is not available, cannot proceed without it"))
+	return nullptr;
+#endif
 }
 
 void FCachedFontData::PrintCache() const
diff --git a/Engine/Plugins/Experimental/Text3D/Source/Text3D/Private/Subsystems/Text3DEngineSubsystem.h b/Engine/Plugins/Experimental/Text3D/Source/Text3D/Private/Subsystems/Text3DEngineSubsystem.h
--- a/Engine/Plugins/Experimental/Text3D/Source/Text3D/Private/Subsystems/Text3DEngineSubsystem.h
+++ b/Engine/Plugins/Experimental/Text3D/Source/Text3D/Private/Subsystems/Text3DEngineSubsystem.h
@@ -106,9 +106,14 @@ struct FTypefaceFontData
 		return Count;
 	}
 
+#if WITH_FREETYPE
 	FT_Face GetTypeFace() const { return TypeFace; }
 	void SetTypeFace(FT_Face InTypeFace);
 
+	TSharedPtr<const FFontFaceData> GetTypeFaceData() const;
+	void SetTypeFaceData(const TSharedPtr<const FFontFaceData>& InTypeFaceData);
+#endif
+
 	FName GetTypefaceName() const { return TypefaceName; }
 	void SetTypefaceName(FName InTypefaceName)
 	{
@@ -121,9 +126,6 @@ struct FTypefaceFontData
 		TypefaceFontDataHash = InDataHash;
 	}
 
-	TSharedPtr<const FFontFaceData> GetTypeFaceData() const;
-	void SetTypeFaceData(const TSharedPtr<const FFontFaceData>& InTypeFaceData);
-
 	TMap<uint32, FCachedFontMeshes>& GetMeshes() { return Meshes; }
 
 	TSharedPtr<int32> GetCacheCounter() { return CacheCounter; }
@@ -140,8 +142,11 @@ private:
 	FName TypefaceName = NAME_None;
 	TSharedPtr<int32> CacheCounter;
 	uint32 TypefaceFontDataHash = 0;
-	FT_Face TypeFace = nullptr;
+
+#if WITH_FREETYPE
 	TWeakPtr<const FFontFaceData> TypeFaceDataWeak;
+	FT_Face TypeFace = nullptr;
+#endif
 };
 
 USTRUCT()
@@ -152,7 +157,9 @@ struct FCachedFontData
 	FCachedFontData();
 	~FCachedFontData();
 
+#if WITH_FREETYPE
 	FT_Face GetFreeTypeFace(uint32 InTypefaceIndex);
+#endif
 
 	FString GetFontName() const;
 
diff --git a/Engine/Plugins/Experimental/Text3D/Source/Text3D/Private/Text3DModule.cpp b/Engine/Plugins/Experimental/Text3D/Source/Text3D/Private/Text3DModule.cpp
--- a/Engine/Plugins/Experimental/Text3D/Source/Text3D/Private/Text3DModule.cpp
+++ b/Engine/Plugins/Experimental/Text3D/Source/Text3D/Private/Text3DModule.cpp
@@ -11,22 +11,30 @@ IMPLEMENT_MODULE(FText3DModule, Text3D);
 
 FText3DModule::FText3DModule()
 {
+#if WITH_FREETYPE
 	FreeTypeLib = nullptr;
+#endif
 }
 
 void FText3DModule::StartupModule()
 {
+#if WITH_FREETYPE
 	FT_Init_FreeType(&FreeTypeLib);
+#endif
 }
 
 void FText3DModule::ShutdownModule()
 {
+#if WITH_FREETYPE
 	FT_Done_FreeType(FreeTypeLib);
 	FreeTypeLib = nullptr;
+#endif
 }
 
+#if WITH_FREETYPE
 FT_Library FText3DModule::GetFreeTypeLibrary()
 {
 	const FText3DModule& Instance = FModuleManager::LoadModuleChecked<FText3DModule>("Text3D");
 	return Instance.FreeTypeLib;
-}
\ No newline at end of file
+}
+#endif
\ No newline at end of file
diff --git a/Engine/Plugins/Experimental/Text3D/Source/Text3D/Private/Text3DModule.h b/Engine/Plugins/Experimental/Text3D/Source/Text3D/Private/Text3DModule.h
--- a/Engine/Plugins/Experimental/Text3D/Source/Text3D/Private/Text3DModule.h
+++ b/Engine/Plugins/Experimental/Text3D/Source/Text3D/Private/Text3DModule.h
@@ -4,6 +4,8 @@
 
 #include "Modules/ModuleInterface.h"
 
+#if WITH_FREETYPE
+
 #if PLATFORM_COMPILER_HAS_GENERIC_KEYWORD
 #define generic __identifier(generic)
 #endif	//PLATFORM_COMPILER_HAS_GENERIC_KEYWORD
@@ -18,6 +20,8 @@ THIRD_PARTY_INCLUDES_END
 #undef generic
 #endif	//PLATFORM_COMPILER_HAS_GENERIC_KEYWORD
 
+#endif
+
 class FText3DModule final : public IModuleInterface
 {
 public:
@@ -28,8 +32,10 @@ public:
 	virtual void ShutdownModule() override;
 	//~ End IModuleInterface
 
+#if WITH_FREETYPE
 	static FT_Library GetFreeTypeLibrary();
 
 private:
 	FT_Library FreeTypeLib;
+#endif
 };
diff --git a/Engine/Plugins/Experimental/Text3D/Source/Text3D/Private/Utilities/Text3DUtilities.cpp b/Engine/Plugins/Experimental/Text3D/Source/Text3D/Private/Utilities/Text3DUtilities.cpp
--- a/Engine/Plugins/Experimental/Text3D/Source/Text3D/Private/Utilities/Text3DUtilities.cpp
+++ b/Engine/Plugins/Experimental/Text3D/Source/Text3D/Private/Utilities/Text3DUtilities.cpp
@@ -5,13 +5,16 @@
 #include "Engine/Font.h"
 #include "Engine/FontFace.h"
 #include "Fonts/CompositeFont.h"
+#include "Logs/Text3DLogs.h"
 #include "Misc/Paths.h"
 #include "Text3DModule.h"
 
+#if WITH_FREETYPE
 THIRD_PARTY_INCLUDES_START
 #include "ft2build.h"
 #include FT_FREETYPE_H
 THIRD_PARTY_INCLUDES_END
+#endif
 
 bool UE::Text3D::Utilities::GetSanitizeFontName(const UFont* InFont, FString& OutFontName)
 {
@@ -81,6 +84,7 @@ void UE::Text3D::Utilities::SanitizeFontName(FString& InOutFontName)
 
 bool UE::Text3D::Utilities::GetFontStyle(const UFont* InFont, EText3DFontStyleFlags& OutFontStyleFlags)
 {
+#if WITH_FREETYPE
 	if (!IsValid(InFont))
 	{
 		return false;
@@ -146,12 +150,17 @@ bool UE::Text3D::Utilities::GetFontStyle(const UFont* InFont, EText3DFontStyleFl
 	}
 
 	return true;
+#else
+	UE_LOG(LogText3D, Warning, TEXT("FreeType is not available, cannot proceed without it"))
+	return false;
+#endif
 }
 
 bool UE::Text3D::Utilities::GetFontStyle(const FText3DFontFamily& InFontFamily, EText3DFontStyleFlags& OutFontStyleFlags)
 {
 	OutFontStyleFlags = EText3DFontStyleFlags::None;
 
+#if WITH_FREETYPE
 	if (InFontFamily.FontFacePaths.IsEmpty())
 	{
 		return false;
@@ -187,6 +196,10 @@ bool UE::Text3D::Utilities::GetFontStyle(const FText3DFontFamily& InFontFamily,
 	}
 
 	return true;
+#else
+	UE_LOG(LogText3D, Warning, TEXT("FreeType is not available, cannot proceed without it"))
+	return false;
+#endif
 }
 
 bool UE::Text3D::Utilities::GetFontFaces(const UFont* InFont, TArray<UFontFace*>& OutFontFaces)
diff --git a/Engine/Plugins/Experimental/Text3D/Source/Text3D/Text3D.Build.cs b/Engine/Plugins/Experimental/Text3D/Source/Text3D/Text3D.Build.cs
--- a/Engine/Plugins/Experimental/Text3D/Source/Text3D/Text3D.Build.cs
+++ b/Engine/Plugins/Experimental/Text3D/Source/Text3D/Text3D.Build.cs
@@ -20,10 +20,8 @@ public class Text3D : ModuleRules
 		PrivateDependencyModuleNames.AddRange(new string[] 
 		{
 			"DeveloperSettings",
-			"FreeType2",
 			"GeometryAlgorithms",
 			"GeometryCore",
-			"HarfBuzz",
             "ICU",
             "MeshDescription",
             "Slate",
@@ -31,6 +29,16 @@ public class Text3D : ModuleRules
 			"StaticMeshDescription",
 		});
 
+		if (Target.Type != TargetType.Server)
+		{
+			if (Target.bCompileFreeType)
+			{
+				PrivateDependencyModuleNames.Add("FreeType2");
+			}
+
+			PrivateDependencyModuleNames.Add("HarfBuzz");
+		}
+
 		// Needed to reference underlying FreeType info
 		// @todo: Move font vector handling to new module ("FontCore") - don't expose FT directly?
 		string EnginePath = Path.GetFullPath(Target.RelativeEnginePath);
