Text3D: Fixed issue where a rebuild request would get queued for next tick but processed with the wrong UpdateFlags (none) when triggered while a build is in progress, since the UpdateFlags would get cleared at the end of the build, and only set when queueing the request.

diff --git a/Engine/Plugins/Experimental/Text3D/Source/Text3D/Private/Text3DComponent.cpp b/Engine/Plugins/Experimental/Text3D/Source/Text3D/Private/Text3DComponent.cpp
--- a/Engine/Plugins/Experimental/Text3D/Source/Text3D/Private/Text3DComponent.cpp
+++ b/Engine/Plugins/Experimental/Text3D/Source/Text3D/Private/Text3DComponent.cpp
@@ -91,26 +91,19 @@ EText3DRendererFlags UText3DComponent::GetUpdateFlags() const
 
 void UText3DComponent::RequestUpdate(EText3DRendererFlags InFlags, bool bInImmediate)
 {
-	EnumAddFlags(UpdateFlags, InFlags);
-
-	if (UpdateFlags != EText3DRendererFlags::None)
+	if (InFlags != EText3DRendererFlags::None)
 	{
 		if (bInImmediate)
 		{
-			RebuildInternal();
+			Rebuild(InFlags, /*bCleanCache*/false);
 		}
 		else
 		{
-			ScheduleTextUpdateNextTick();
+			QueueRebuild(InFlags);
 		}
 	}
 }
 
-void UText3DComponent::ClearUpdateFlags()
-{
-	UpdateFlags = EText3DRendererFlags::None;
-}
-
 uint32 UText3DComponent::GetTypeFaceIndex() const
 {
 	uint32 TypefaceIndex = INDEX_NONE;
@@ -876,6 +869,51 @@ void UText3DComponent::SetTextRendererClass(const TSubclassOf<UText3DRendererBas
 
 void UText3DComponent::RebuildInternal(bool bCleanCache)
 {
+}
+
+void UText3DComponent::QueueRebuild(EText3DRendererFlags InFlags)
+{
+	// Hotfix: To not break ABI compatibility, keep track of the queued flags here
+	static TMap<FObjectKey, EText3DRendererFlags> QueuedUpdateFlags;
+
+	// Add the given flags to the current queued flags in case a text update is already queued
+	EText3DRendererFlags& CurrentUpdateFlags = QueuedUpdateFlags.FindOrAdd(this, EText3DRendererFlags::None);
+	EnumAddFlags(CurrentUpdateFlags, InFlags);
+
+	if (!TextUpdateHandle.IsValid())
+	{
+		TextUpdateHandle = FTSTicker::GetCoreTicker().AddTicker(
+			FTickerDelegate::CreateWeakLambda(
+				this,
+				[this](float)->bool
+				{
+					if (UE::IsSavingPackage(this))
+					{
+						// Postpone to next tick
+						return true;
+					}
+
+					// Reset handle so that while building text, new rebuild requests can be issued
+					// New rebuild requests should only happen when, while building text, existing 'built data' becomes outdated.
+					TextUpdateHandle.Reset();
+
+					EText3DRendererFlags LatestUpdateFlags = EText3DRendererFlags::None;
+					QueuedUpdateFlags.RemoveAndCopyValue(this, LatestUpdateFlags);
+
+					Rebuild(LatestUpdateFlags, /*bCleanCache*/false);
+
+					// Executes only once
+					return false;
+				}
+			)
+		);
+	}
+}
+
+void UText3DComponent::Rebuild(EText3DRendererFlags InUpdateFlags, bool bCleanCache)
+{
+	EnumAddFlags(UpdateFlags, InUpdateFlags);
+
 	if (IsTemplate(RF_ClassDefaultObject))
 	{
 		return;
@@ -913,11 +951,18 @@ void UText3DComponent::RebuildInternal(bool bCleanCache)
 	
 	if (bIsUpdatingText)
 	{
+		UE_LOG(LogText3D, Warning, TEXT("[%s] Rebuild called with update flags (%u) while building text. This new Rebuild will be deferred to the next tick."), *GetReadableName(), static_cast<uint8>(UpdateFlags));
+		QueueRebuild(InUpdateFlags);
 		return;
 	}
 
 	bIsUpdatingText = true;
-	ON_SCOPE_EXIT { bIsUpdatingText = false; };
+
+	ON_SCOPE_EXIT
+	{
+		bIsUpdatingText = false;
+		EnumRemoveFlags(UpdateFlags, InUpdateFlags);
+	};
 
 	TextPreUpdateDelegate.Broadcast(this, UpdateFlags);
 
@@ -959,8 +1004,10 @@ void UText3DComponent::RebuildInternal(bool bCleanCache)
 	}
 
 	TextPostUpdateDelegate.Broadcast(this, UpdateFlags);
+}
 
-	ClearUpdateFlags();
+void UText3DComponent::ClearUpdateFlags()
+{
 }
 
 void UText3DComponent::OnTextRendererClassChanged()
@@ -1021,29 +1068,6 @@ void UText3DComponent::OnFontPropertiesChanged()
 
 void UText3DComponent::ScheduleTextUpdateNextTick()
 {
-	if (!TextUpdateHandle.IsValid())
-	{
-		TextUpdateHandle = FTSTicker::GetCoreTicker().AddTicker(
-			FTickerDelegate::CreateWeakLambda(
-				this,
-				[this](float)->bool
-				{
-					if (UE::IsSavingPackage(this))
-					{
-						// Postpone to next tick
-						return true;
-					}
-
-					RebuildInternal();
-
-					TextUpdateHandle.Reset();
-
-					// Executes only once
-					return false;
-				}
-			)
-		);
-	}
 }
 
 TArray<FName> UText3DComponent::GetTypefaceNames() const
diff --git a/Engine/Plugins/Experimental/Text3D/Source/Text3D/Public/Text3DComponent.h b/Engine/Plugins/Experimental/Text3D/Source/Text3D/Public/Text3DComponent.h
--- a/Engine/Plugins/Experimental/Text3D/Source/Text3D/Public/Text3DComponent.h
+++ b/Engine/Plugins/Experimental/Text3D/Source/Text3D/Public/Text3DComponent.h
@@ -422,7 +422,7 @@ protected:
 	/** Called when font options are changed */
 	void OnFontPropertiesChanged();
 
-	/** Plan update task on next tick */
+	/** Todo: Remove for 5.7. Kept here for ABI Compatibility */
 	void ScheduleTextUpdateNextTick();
 
 	/** The renderer class to use to create the text geometry */
@@ -495,8 +495,16 @@ private:
 	/** Planned update requested */
 	FTSTicker::FDelegateHandle TextUpdateHandle;
 
+	/** Todo: Remove for 5.7. Kept here for ABI Compatibility */
 	void RebuildInternal(bool bCleanCache = false);
 
+	/** Queues text rebuild for the next tick */
+	void QueueRebuild(EText3DRendererFlags InFlags);
+
+	/** Partially or fully rebuilds text depending on the given update flags */
+	void Rebuild(EText3DRendererFlags InUpdateFlags, bool bCleanCache);
+
+	/** Todo: Remove for 5.7. Kept here for ABI Compatibility */
 	void ClearUpdateFlags();
 
 	bool IsTypefaceAvailable(FName InTypeface) const;
