MotionDesign : Text3D
- Added debug print for contour node

diff --git a/Engine/Plugins/Experimental/Text3D/Source/Text3D/Private/GeometryBuilders/Text3DGlyphContourNode.h b/Engine/Plugins/Experimental/Text3D/Source/Text3D/Private/GeometryBuilders/Text3DGlyphContourNode.h
--- a/Engine/Plugins/Experimental/Text3D/Source/Text3D/Private/GeometryBuilders/Text3DGlyphContourNode.h
+++ b/Engine/Plugins/Experimental/Text3D/Source/Text3D/Private/GeometryBuilders/Text3DGlyphContourNode.h
@@ -3,8 +3,9 @@
 #pragma once
 
 #include "Containers/Array.h"
-#include "Templates/SharedPointer.h"
+#include "Logs/Text3DLogs.h"
 #include "Polygon2.h"
+#include "Templates/SharedPointer.h"
 
 using UE::Geometry::FPolygon2f;
 
@@ -29,4 +30,22 @@ struct FText3DGlyphContourNode final
 
 	// Contours that are inside this contour
 	TArray<TText3DGlyphContourNodeShared> Children;
+
+	/** Debug print to display contour data */
+	void Print(int32 InDepth = 0) const
+	{
+		const FString Indent = FString::ChrN(InDepth * 2, TEXT(' '));
+
+		UE_LOG(LogText3D, Log, TEXT("%sContour: Vtx=%d, Area=%f, Bounds=%s, CW=%s"),
+			*Indent,
+			Contour.IsValid() ? Contour->VertexCount() : 0,
+			Contour.IsValid() ? Contour->SignedArea() : 0,
+			Contour.IsValid() ? *Contour->Bounds().Extents().ToString() : TEXT(""),
+			bClockwise ? TEXT("true") : TEXT("false"));
+
+		for (const TSharedPtr<FText3DGlyphContourNode>& Child : Children)
+		{
+			Child->Print(InDepth + 1);
+		}
+	}
 };
