MotionDesign : Text3D
- Fixed missing glyph geometry when using some fonts like Antonio, Teko... the cause is the code checking if a contour is contained within another only tests a segments instead of the entire polygon or bounding box

diff --git a/Engine/Plugins/Experimental/Text3D/Source/Text3D/Private/GeometryBuilders/Text3DGlyphLoader.cpp b/Engine/Plugins/Experimental/Text3D/Source/Text3D/Private/GeometryBuilders/Text3DGlyphLoader.cpp
--- a/Engine/Plugins/Experimental/Text3D/Source/Text3D/Private/GeometryBuilders/Text3DGlyphLoader.cpp
+++ b/Engine/Plugins/Experimental/Text3D/Source/Text3D/Private/GeometryBuilders/Text3DGlyphLoader.cpp
@@ -430,10 +430,9 @@ TSharedPtr<FPolygon2f> FText3DGlyphLoader::ProcessFreetypeOutline(const FT_Outli
 }
 #endif
 
-bool FText3DGlyphLoader::IsInside(const TSharedPtr<const FPolygon2f> ContourA,
-							const TSharedPtr<const FPolygon2f> ContourB) const
+bool FText3DGlyphLoader::IsInside(const TSharedPtr<const FPolygon2f> ContourA, const TSharedPtr<const FPolygon2f> ContourB) const
 {
-	return ContourB->Contains(ContourA->GetSegmentPoint(0, 0));
+	return ContourB->Contains(*ContourA) || ContourB->Bounds().Contains(ContourA->Bounds());
 }
 
 void FText3DGlyphLoader::Insert(const TText3DGlyphContourNodeShared NodeA, const TText3DGlyphContourNodeShared NodeB) const
