MotionDesign : Text3D
- Debug component tree in details panel for static mesh renderer
- Recycle components instead of creating/destroying using a pool
- Components and renderer are no longer transient
- Reducing amount of components needed by 2
- Suppressing material & mesh warning when using Text3DComponent in a blueprint

diff --git a/Engine/Plugins/Experimental/Text3D/Source/Text3D/Private/Extensions/Text3DDefaultGeometryExtension.cpp b/Engine/Plugins/Experimental/Text3D/Source/Text3D/Private/Extensions/Text3DDefaultGeometryExtension.cpp
--- a/Engine/Plugins/Experimental/Text3D/Source/Text3D/Private/Extensions/Text3DDefaultGeometryExtension.cpp
+++ b/Engine/Plugins/Experimental/Text3D/Source/Text3D/Private/Extensions/Text3DDefaultGeometryExtension.cpp
@@ -169,9 +169,9 @@ FTypefaceFontData* UText3DDefaultGeometryExtension::GetTypefaceFontData() const
 	return nullptr;
 }
 
-EText3DExtensionResult UText3DDefaultGeometryExtension::PreRendererUpdate(EText3DRendererFlags InFlag)
+EText3DExtensionResult UText3DDefaultGeometryExtension::PreRendererUpdate(const UE::Text3D::Renderer::FUpdateParameters& InParameters)
 {
-	if (InFlag != EText3DRendererFlags::Geometry)
+	if (InParameters.CurrentFlag != EText3DRendererFlags::Geometry)
 	{
 		return EText3DExtensionResult::Active;
 	}
@@ -205,7 +205,7 @@ EText3DExtensionResult UText3DDefaultGeometryExtension::PreRendererUpdate(EText3
 	return EText3DExtensionResult::Active;
 }
 
-EText3DExtensionResult UText3DDefaultGeometryExtension::PostRendererUpdate(EText3DRendererFlags InFlag)
+EText3DExtensionResult UText3DDefaultGeometryExtension::PostRendererUpdate(const UE::Text3D::Renderer::FUpdateParameters& InParameters)
 {
 	CachedFontData = nullptr;
 	return EText3DExtensionResult::Finished;
diff --git a/Engine/Plugins/Experimental/Text3D/Source/Text3D/Private/Extensions/Text3DDefaultLayoutExtension.cpp b/Engine/Plugins/Experimental/Text3D/Source/Text3D/Private/Extensions/Text3DDefaultLayoutExtension.cpp
--- a/Engine/Plugins/Experimental/Text3D/Source/Text3D/Private/Extensions/Text3DDefaultLayoutExtension.cpp
+++ b/Engine/Plugins/Experimental/Text3D/Source/Text3D/Private/Extensions/Text3DDefaultLayoutExtension.cpp
@@ -181,9 +181,9 @@ FVector UText3DDefaultLayoutExtension::GetTextScale() const
 	return TextScale;
 }
 
-EText3DExtensionResult UText3DDefaultLayoutExtension::PreRendererUpdate(EText3DRendererFlags InFlag)
+EText3DExtensionResult UText3DDefaultLayoutExtension::PreRendererUpdate(const UE::Text3D::Renderer::FUpdateParameters& InParameters)
 {
-	if (InFlag != EText3DRendererFlags::Geometry && InFlag != EText3DRendererFlags::Layout)
+	if (InParameters.CurrentFlag != EText3DRendererFlags::Geometry && InParameters.CurrentFlag != EText3DRendererFlags::Layout)
 	{
 		return EText3DExtensionResult::Active;
 	}
@@ -210,7 +210,7 @@ EText3DExtensionResult UText3DDefaultLayoutExtension::PreRendererUpdate(EText3DR
 		FontStyle = MakeShared<FTextBlockStyle>();
 	}
 
-	if (EnumHasAnyFlags(Text3DComponent->GetUpdateFlags(), EText3DRendererFlags::Geometry))
+	if (EnumHasAnyFlags(InParameters.UpdateFlags, EText3DRendererFlags::Geometry))
 	{
 		UFont* Font = Text3DComponent->GetFont();
 		FSlateFontInfo FontInfo(Font, UE::Text3D::Metrics::SlateFontSize);
@@ -230,7 +230,7 @@ EText3DExtensionResult UText3DDefaultLayoutExtension::PreRendererUpdate(EText3DR
 		});
 	},0);
 
-	if (EnumHasAnyFlags(Text3DComponent->GetUpdateFlags(), EText3DRendererFlags::Geometry))
+	if (EnumHasAnyFlags(InParameters.UpdateFlags, EText3DRendererFlags::Geometry))
 	{
 		UText3DCharacterExtensionBase* CharacterExtension = Text3DComponent->GetCharacterExtension();
 		CharacterExtension->AllocateCharacters(CharacterCount);
@@ -277,7 +277,7 @@ EText3DExtensionResult UText3DDefaultLayoutExtension::PreRendererUpdate(EText3DR
 		}
 	}
 
-	if (EnumHasAnyFlags(Text3DComponent->GetUpdateFlags(), EText3DRendererFlags::Geometry))
+	if (EnumHasAnyFlags(InParameters.UpdateFlags, EText3DRendererFlags::Geometry))
 	{
 		FTypefaceFontData* TypefaceFontData = GeometryExtension->GetTypefaceFontData();
 		if (!TypefaceFontData)
@@ -319,7 +319,7 @@ EText3DExtensionResult UText3DDefaultLayoutExtension::PreRendererUpdate(EText3DR
 				const uint16 Index = bIsRTL ? CharacterCount - 1 - CharacterIndex : CharacterIndex;
 				UText3DCharacterBase* Character = Text3DComponent->GetCharacter(Index);
 
-				if (EnumHasAnyFlags(Text3DComponent->GetUpdateFlags(), EText3DRendererFlags::Geometry))
+				if (EnumHasAnyFlags(InParameters.UpdateFlags, EText3DRendererFlags::Geometry))
 				{
 					Character->SetGlyphIndex(ShapedGlyph->GlyphIndex);
 					GeometryExtension->FindOrLoadGlyphMesh(Character);
@@ -359,7 +359,7 @@ EText3DExtensionResult UText3DDefaultLayoutExtension::PreRendererUpdate(EText3DR
 	return EText3DExtensionResult::Finished;
 }
 
-EText3DExtensionResult UText3DDefaultLayoutExtension::PostRendererUpdate(EText3DRendererFlags InFlag)
+EText3DExtensionResult UText3DDefaultLayoutExtension::PostRendererUpdate(const UE::Text3D::Renderer::FUpdateParameters& InParameters)
 {
 	return EText3DExtensionResult::Active;
 }
diff --git a/Engine/Plugins/Experimental/Text3D/Source/Text3D/Private/Extensions/Text3DDefaultMaterialExtension.cpp b/Engine/Plugins/Experimental/Text3D/Source/Text3D/Private/Extensions/Text3DDefaultMaterialExtension.cpp
--- a/Engine/Plugins/Experimental/Text3D/Source/Text3D/Private/Extensions/Text3DDefaultMaterialExtension.cpp
+++ b/Engine/Plugins/Experimental/Text3D/Source/Text3D/Private/Extensions/Text3DDefaultMaterialExtension.cpp
@@ -276,9 +276,9 @@ void UText3DDefaultMaterialExtension::SetBackMaterial(UMaterialInterface* InMate
 	OnCustomMaterialChanged();
 }
 
-EText3DExtensionResult UText3DDefaultMaterialExtension::PreRendererUpdate(EText3DRendererFlags InFlag)
+EText3DExtensionResult UText3DDefaultMaterialExtension::PreRendererUpdate(const UE::Text3D::Renderer::FUpdateParameters& InParameters)
 {
-	if (InFlag != EText3DRendererFlags::Material)
+	if (InParameters.CurrentFlag != EText3DRendererFlags::Material)
 	{
 		return EText3DExtensionResult::Active;
 	}
@@ -368,7 +368,7 @@ EText3DExtensionResult UText3DDefaultMaterialExtension::PreRendererUpdate(EText3
 	return EText3DExtensionResult::Finished;
 }
 
-EText3DExtensionResult UText3DDefaultMaterialExtension::PostRendererUpdate(EText3DRendererFlags InFlag)
+EText3DExtensionResult UText3DDefaultMaterialExtension::PostRendererUpdate(const UE::Text3D::Renderer::FUpdateParameters& InParameters)
 {
 	return EText3DExtensionResult::Active;
 }
@@ -421,7 +421,8 @@ UMaterialInstanceDynamic* UText3DDefaultMaterialExtension::FindOrAdd(EText3DGrou
 		return *Material;
 	}
 
-	UMaterialInstanceDynamic* DynamicMaterial = UMaterialInstanceDynamic::Create(ParentMaterial, GetTransientPackage());
+	UMaterialInstanceDynamic* DynamicMaterial = UMaterialInstanceDynamic::Create(ParentMaterial, this);
+	DynamicMaterial->SetFlags(RF_Transient | RF_Public);
 	GroupDynamicMaterials.Emplace(GroupKey, DynamicMaterial);
 
 	return DynamicMaterial;
diff --git a/Engine/Plugins/Experimental/Text3D/Source/Text3D/Private/Extensions/Text3DEffectExtensionBase.cpp b/Engine/Plugins/Experimental/Text3D/Source/Text3D/Private/Extensions/Text3DEffectExtensionBase.cpp
--- a/Engine/Plugins/Experimental/Text3D/Source/Text3D/Private/Extensions/Text3DEffectExtensionBase.cpp
+++ b/Engine/Plugins/Experimental/Text3D/Source/Text3D/Private/Extensions/Text3DEffectExtensionBase.cpp
@@ -3,9 +3,9 @@
 #include "Extensions/Text3DEffectExtensionBase.h"
 #include "Text3DComponent.h"
 
-EText3DExtensionResult UText3DEffectExtensionBase::PreRendererUpdate(EText3DRendererFlags InFlag)
+EText3DExtensionResult UText3DEffectExtensionBase::PreRendererUpdate(const UE::Text3D::Renderer::FUpdateParameters& InParameters)
 {
-	if (InFlag != EText3DRendererFlags::Layout)
+	if (InParameters.CurrentFlag != EText3DRendererFlags::Layout)
 	{
 		return EText3DExtensionResult::Active;
 	}
@@ -24,7 +24,7 @@ EText3DExtensionResult UText3DEffectExtensionBase::PreRendererUpdate(EText3DRend
 	return EText3DExtensionResult::Finished;
 }
 
-EText3DExtensionResult UText3DEffectExtensionBase::PostRendererUpdate(EText3DRendererFlags InFlag)
+EText3DExtensionResult UText3DEffectExtensionBase::PostRendererUpdate(const UE::Text3D::Renderer::FUpdateParameters& InParameters)
 {
 	return EText3DExtensionResult::Active;
 }
\ No newline at end of file
diff --git a/Engine/Plugins/Experimental/Text3D/Source/Text3D/Private/Renderers/StaticMeshes/Text3DStaticMeshesRenderer.cpp b/Engine/Plugins/Experimental/Text3D/Source/Text3D/Private/Renderers/StaticMeshes/Text3DStaticMeshesRenderer.cpp
--- a/Engine/Plugins/Experimental/Text3D/Source/Text3D/Private/Renderers/StaticMeshes/Text3DStaticMeshesRenderer.cpp
+++ b/Engine/Plugins/Experimental/Text3D/Source/Text3D/Private/Renderers/StaticMeshes/Text3DStaticMeshesRenderer.cpp
@@ -3,7 +3,6 @@
 #include "Renderers/StaticMeshes/Text3DStaticMeshesRenderer.h"
 
 #include "Algo/Accumulate.h"
-#include "Algo/Count.h"
 #include "Characters/Text3DCharacterBase.h"
 #include "Components/SceneComponent.h"
 #include "Components/StaticMeshComponent.h"
@@ -12,11 +11,6 @@
 #include "Extensions/Text3DLayoutExtensionBase.h"
 #include "Extensions/Text3DMaterialExtensionBase.h"
 #include "Extensions/Text3DRenderingExtensionBase.h"
-#include "LayoutBuilders/Text3DShapedGlyphText.h"
-#include "LayoutBuilders/Text3DLayoutShaper.h"
-#include "Logs/Text3DLogs.h"
-#include "Misc/EnumerateRange.h"
-#include "Subsystems/Text3DEngineSubsystem.h"
 #include "Text3DComponent.h"
 #include "Text3DInternalTypes.h"
 
@@ -25,18 +19,18 @@ void UText3DStaticMeshesRenderer::OnCreate()
 	if (!TextRoot)
 	{
 		UText3DComponent* TextComponent = GetText3DComponent();
-		TextRoot = NewObject<USceneComponent>(TextComponent, TEXT("TextRoot"), RF_Transient);
+		TextRoot = NewObject<USceneComponent>(this, TEXT("TextRoot"));
 		TextRoot->SetupAttachment(TextComponent);
 	}
 }
 
-void UText3DStaticMeshesRenderer::OnUpdate(EText3DRendererFlags InFlags)
+void UText3DStaticMeshesRenderer::OnUpdate(const UE::Text3D::Renderer::FUpdateParameters& InParameters)
 {
 	const UText3DComponent* TextComponent = GetText3DComponent();
 
-	if (EnumHasAnyFlags(InFlags, EText3DRendererFlags::Geometry))
+	if (InParameters.CurrentFlag == EText3DRendererFlags::Geometry)
 	{
-		AllocateCharacterComponents(TextComponent->GetCharacterCount());
+		AllocateComponents(TextComponent->GetCharacterCount());
 
 		const UText3DGeometryExtensionBase* GeometryExtension = TextComponent->GetGeometryExtension();
 
@@ -48,8 +42,7 @@ void UText3DStaticMeshesRenderer::OnUpdate(EText3DRendererFlags InFlags)
 
 		RefreshBounds();
 	}
-
-	if (EnumHasAnyFlags(InFlags, EText3DRendererFlags::Layout))
+	else if (InParameters.CurrentFlag == EText3DRendererFlags::Layout)
 	{
 		const UText3DLayoutExtensionBase* LayoutExtension = TextComponent->GetLayoutExtension();
 
@@ -59,14 +52,12 @@ void UText3DStaticMeshesRenderer::OnUpdate(EText3DRendererFlags InFlags)
 		{
 			constexpr bool bReset = false;
 			const FTransform& CharacterTransform = InCharacter->GetTransform(bReset);
-			CharacterKernings[InIndex]->SetRelativeTransform(FTransform::Identity);
 			CharacterMeshes[InIndex]->SetRelativeTransform(CharacterTransform);
 		});
 
 		RefreshBounds();
 	}
-
-	if (EnumHasAnyFlags(InFlags, EText3DRendererFlags::Material))
+	else if (InParameters.CurrentFlag == EText3DRendererFlags::Material)
 	{
 		using namespace UE::Text3D::Materials;
 
@@ -93,8 +84,7 @@ void UText3DStaticMeshesRenderer::OnUpdate(EText3DRendererFlags InFlags)
 			}
 		});
 	}
-
-	if (EnumHasAnyFlags(InFlags, EText3DRendererFlags::Visibility))
+	else if (InParameters.CurrentFlag == EText3DRendererFlags::Visibility)
 	{
 		const UText3DRenderingExtensionBase* RenderingExtension = TextComponent->GetRenderingExtension();
 
@@ -113,47 +103,31 @@ void UText3DStaticMeshesRenderer::OnUpdate(EText3DRendererFlags InFlags)
 }
 
 void UText3DStaticMeshesRenderer::OnClear()
+{
+	AllocateComponents(0);
+}
+
+void UText3DStaticMeshesRenderer::OnDestroy()
 {
 	for (UStaticMeshComponent* MeshComponent : CharacterMeshes)
 	{
 		if (MeshComponent)
 		{
 			MeshComponent->DetachFromComponent(FDetachmentTransformRules::KeepRelativeTransform);
-			MeshComponent->SetStaticMesh(nullptr);
 			MeshComponent->DestroyComponent();
 		}
 	}
 	CharacterMeshes.Reset();
 
-	for (USceneComponent* KerningComponent : CharacterKernings)
+	for (UStaticMeshComponent* MeshComponent : CharacterMeshesPool)
 	{
-		if (KerningComponent)
+		if (MeshComponent)
 		{
-			KerningComponent->DetachFromComponent(FDetachmentTransformRules::KeepRelativeTransform);
-			KerningComponent->DestroyComponent();
+			MeshComponent->DestroyComponent();
 		}
 	}
-	CharacterKernings.Reset();
+	CharacterMeshes.Reset();
 
-	if (TextRoot)
-	{
-		TArray<USceneComponent*> ChildComponents;
-		constexpr bool bIncludeChildDescendants = true;
-		TextRoot->GetChildrenComponents(bIncludeChildDescendants, ChildComponents);
-
-		for (USceneComponent* ChildComponent : ChildComponents)
-		{
-			if (IsValid(ChildComponent))
-			{
-				ChildComponent->DetachFromComponent(FDetachmentTransformRules::KeepRelativeTransform);
-				ChildComponent->DestroyComponent();
-			}
-		}
-	}
-}
-
-void UText3DStaticMeshesRenderer::OnDestroy()
-{
 	if (TextRoot)
 	{
 		TextRoot->DestroyComponent();
@@ -161,7 +135,7 @@ void UText3DStaticMeshesRenderer::OnDestroy()
 	}
 }
 
-FName UText3DStaticMeshesRenderer::GetName() const
+FName UText3DStaticMeshesRenderer::GetFriendlyName() const
 {
 	static const FName Name(TEXT("StaticMeshesRenderer"));
 	return Name;
@@ -179,64 +153,55 @@ FBox UText3DStaticMeshesRenderer::OnCalculateBounds() const
 	return Box;
 }
 
-bool UText3DStaticMeshesRenderer::AllocateCharacterComponents(int32 Num)
+void UText3DStaticMeshesRenderer::AllocateComponents(int32 InCount)
 {
-	int32 DeltaNum = Num - CharacterMeshes.Num();
-	if (DeltaNum == 0)
+	if (CharacterMeshes.Num() == InCount)
 	{
-		return false;
+		return;
 	}
 
-	// Add characters
-	if (FMath::Sign(DeltaNum) > 0)
+	CharacterMeshes.Reserve(InCount);
+
+	const int32 RemainingCharacterCount = CharacterMeshes.Num() - InCount;
+
+	AActor* Actor = GetTypedOuter<AActor>();
+	if (RemainingCharacterCount > 0)
 	{
-		int32 GlyphId = CharacterMeshes.Num() - 1;
-		for(int32 CharacterIndex = 0; CharacterIndex < DeltaNum; ++CharacterIndex)
+		for (int32 CharacterIndex = InCount; CharacterIndex < CharacterMeshes.Num(); CharacterIndex++)
 		{
-			GlyphId++;
-
-			const FName CharacterKerningComponentName = MakeUniqueObjectName(this, USceneComponent::StaticClass(), FName(*FString::Printf(TEXT("CharacterKerning%d"), GlyphId)));
-			USceneComponent* CharacterKerningComponent = NewObject<USceneComponent>(this, CharacterKerningComponentName, RF_Transient);
-			CharacterKerningComponent->AttachToComponent(TextRoot, FAttachmentTransformRules::KeepRelativeTransform);
-			CharacterKerningComponent->RegisterComponent();
-			CharacterKernings.Add(CharacterKerningComponent);
-
-			const FName StaticMeshComponentName = MakeUniqueObjectName(this, UStaticMeshComponent::StaticClass(), FName(*FString::Printf(TEXT("StaticMeshComponent%d"), GlyphId)));
-			UStaticMeshComponent* StaticMeshComponent = NewObject<UStaticMeshComponent>(this, StaticMeshComponentName, RF_Transient);
-			StaticMeshComponent->AttachToComponent(CharacterKerningComponent, FAttachmentTransformRules::KeepRelativeTransform);
-			StaticMeshComponent->RegisterComponent();
-			CharacterMeshes.Add(StaticMeshComponent);
+			if (UStaticMeshComponent* MeshComponent = CharacterMeshes[CharacterIndex])
+			{
+				MeshComponent->DetachFromComponent(FDetachmentTransformRules::KeepRelativeTransform);
+				MeshComponent->UnregisterComponent();
+				MeshComponent->SetStaticMesh(nullptr);
+				Actor->RemoveInstanceComponent(MeshComponent);
+				Actor->RemoveOwnedComponent(MeshComponent);
+				CharacterMeshesPool.Add(MeshComponent);
+			}
 		}
 	}
-	// Remove characters
-	else
+	else if (RemainingCharacterCount < 0)
 	{
-		DeltaNum = FMath::Abs(DeltaNum);
-		for(int32 CharacterIndex = CharacterKernings.Num() - 1 - DeltaNum; CharacterIndex < CharacterKernings.Num(); ++CharacterIndex)
+		for (int32 CharacterIndex = 0; CharacterIndex < FMath::Abs(RemainingCharacterCount); ++CharacterIndex)
 		{
-			USceneComponent* CharacterKerningComponent = CharacterKernings[CharacterIndex];
-			// If called in quick succession, may already be pending destruction
-			if (IsValid(CharacterKerningComponent))
+			UStaticMeshComponent* MeshComponent = nullptr;
+			if (!CharacterMeshesPool.IsEmpty())
 			{
-				CharacterKerningComponent->DetachFromComponent(FDetachmentTransformRules::KeepRelativeTransform);
-				CharacterKerningComponent->UnregisterComponent();
-				CharacterKerningComponent->DestroyComponent();
+				MeshComponent = CharacterMeshes.Add_GetRef(CharacterMeshesPool.Pop());
 			}
-
-			UStaticMeshComponent* StaticMeshComponent = CharacterMeshes[CharacterIndex];
-			if (IsValid(StaticMeshComponent))
+			else
 			{
-				StaticMeshComponent->DetachFromComponent(FDetachmentTransformRules::KeepRelativeTransform);
-				StaticMeshComponent->UnregisterComponent();
-				StaticMeshComponent->DestroyComponent();
+				const FName StaticMeshComponentName = MakeUniqueObjectName(this, UStaticMeshComponent::StaticClass(), FName(TEXT("CharacterMesh")));
+				MeshComponent = NewObject<UStaticMeshComponent>(this, StaticMeshComponentName);
+				CharacterMeshes.Add(MeshComponent);
 			}
-		}
 
-		CharacterKernings.RemoveAt(CharacterKernings.Num() - 1 - DeltaNum, DeltaNum);
-		CharacterMeshes.RemoveAt(CharacterMeshes.Num() - 1 - DeltaNum, DeltaNum);
+			MeshComponent->AttachToComponent(TextRoot, FAttachmentTransformRules::KeepRelativeTransform);
+			MeshComponent->RegisterComponent();
+		}
 	}
 
-	return true;
+	CharacterMeshes.SetNum(InCount);
 }
 
 int32 UText3DStaticMeshesRenderer::GetGlyphCount()
@@ -249,21 +214,6 @@ int32 UText3DStaticMeshesRenderer::GetGlyphCount()
 	return TextRoot->GetNumChildrenComponents();
 }
 
-USceneComponent* UText3DStaticMeshesRenderer::GetGlyphKerningComponent(int32 Index)
-{
-	if (!CharacterKernings.IsValidIndex(Index))
-	{
-		return nullptr;
-	}
-
-	return CharacterKernings[Index];
-}
-
-const TArray<USceneComponent*>& UText3DStaticMeshesRenderer::GetGlyphKerningComponents()
-{
-	return CharacterKernings;
-}
-
 UStaticMeshComponent* UText3DStaticMeshesRenderer::GetGlyphMeshComponent(int32 Index)
 {
 	if (!CharacterMeshes.IsValidIndex(Index))
diff --git a/Engine/Plugins/Experimental/Text3D/Source/Text3D/Private/Renderers/StaticMeshes/Text3DStaticMeshesRenderer.h b/Engine/Plugins/Experimental/Text3D/Source/Text3D/Private/Renderers/StaticMeshes/Text3DStaticMeshesRenderer.h
--- a/Engine/Plugins/Experimental/Text3D/Source/Text3D/Private/Renderers/StaticMeshes/Text3DStaticMeshesRenderer.h
+++ b/Engine/Plugins/Experimental/Text3D/Source/Text3D/Private/Renderers/StaticMeshes/Text3DStaticMeshesRenderer.h
@@ -17,12 +17,11 @@ class UText3DComponent;
  * Kerning is done through a scene component containing the current character
  * Text3DComponent
  * - Text3DRoot (Root)
- * -- Text3DSceneComponent (Kerning)
  * --- Text3DStaticMeshComponent (Character)
  *
  * Eg: The text "Hello" will be rendered using 5 SceneComponents and 5 StaticMeshComponents
  */
-UCLASS(DisplayName="StaticMeshesRenderer", ClassGroup="Text3D")
+UCLASS(DisplayName="StaticMeshesRenderer", ClassGroup="Text3D", HideCategories=(Renderer))
 class UText3DStaticMeshesRenderer : public UText3DRendererBase
 {
 	GENERATED_BODY()
@@ -32,14 +31,6 @@ public:
 	UFUNCTION(BlueprintCallable, Category = "Rendering|Components|Text3D")
 	int32 GetGlyphCount();
 
-	/** Gets the USceneComponent that a glyph is attached to */
-	UFUNCTION(BlueprintCallable, Category = "Rendering|Components|Text3D")
-	USceneComponent* GetGlyphKerningComponent(int32 Index);
-
-	/** Gets all the glyph kerning components */
-	UFUNCTION(BlueprintCallable, Category = "Rendering|Components|Text3D")
-	const TArray<USceneComponent*>& GetGlyphKerningComponents();
-
 	/** Gets the StaticMeshComponent of a glyph */
 	UFUNCTION(BlueprintCallable, Category = "Rendering|Components|Text3D")
 	UStaticMeshComponent* GetGlyphMeshComponent(int32 Index);
@@ -51,25 +42,24 @@ public:
 protected:
 	//~ Begin UText3DRendererBase
 	virtual void OnCreate() override;
-	virtual void OnUpdate(EText3DRendererFlags InFlags) override;
+	virtual void OnUpdate(const UE::Text3D::Renderer::FUpdateParameters& InParameters) override;
 	virtual void OnClear() override;
 	virtual void OnDestroy() override;
-	virtual FName GetName() const override;
+	virtual FName GetFriendlyName() const override;
 	virtual FBox OnCalculateBounds() const override;
 	//~ End UText3DRendererBase
 
-	/** Allocates, or shrinks existing components to match the input number. Returns false if nothing modified. */
-	bool AllocateCharacterComponents(int32 Num);
+	void AllocateComponents(int32 InCount);
 
 	/** Holds all character components */
-	UPROPERTY(Transient)
+	UPROPERTY(VisibleAnywhere, DuplicateTransient, TextExportTransient, Category="Renderer", meta=(EditCondition="false", EditConditionHides))
 	TObjectPtr<USceneComponent> TextRoot;
 
-	/** Each character kerning is held in these components */
-	UPROPERTY(Transient)
-	TArray<TObjectPtr<USceneComponent>> CharacterKernings;
-
 	/** Each character mesh is held in these components */
-	UPROPERTY(Transient)
+	UPROPERTY(VisibleAnywhere, DuplicateTransient, TextExportTransient, Category="Renderer", meta=(EditCondition="false", EditConditionHides))
 	TArray<TObjectPtr<UStaticMeshComponent>> CharacterMeshes;
+
+	/** Reusable components to avoid creating and destroying them, recycle components or create new ones when pool is empty */
+	UPROPERTY(Transient, DuplicateTransient, TextExportTransient)
+	TArray<TObjectPtr<UStaticMeshComponent>> CharacterMeshesPool;
 };
\ No newline at end of file
diff --git a/Engine/Plugins/Experimental/Text3D/Source/Text3D/Private/Renderers/Text3DRendererBase.cpp b/Engine/Plugins/Experimental/Text3D/Source/Text3D/Private/Renderers/Text3DRendererBase.cpp
--- a/Engine/Plugins/Experimental/Text3D/Source/Text3D/Private/Renderers/Text3DRendererBase.cpp
+++ b/Engine/Plugins/Experimental/Text3D/Source/Text3D/Private/Renderers/Text3DRendererBase.cpp
@@ -37,11 +37,13 @@ void UText3DRendererBase::Create()
 
 	OnCreate();
 	bInitialized = true;
-	UE_LOG(LogText3D, Log, TEXT("%s : Text3DRenderer %s Created"), *Text3DComponent->GetOwner()->GetActorNameOrLabel(), *GetName().ToString())
+	UE_LOG(LogText3D, Log, TEXT("%s : Text3DRenderer %s Created"), *Text3DComponent->GetOwner()->GetActorNameOrLabel(), *GetFriendlyName().ToString())
 }
 
 void UText3DRendererBase::Update(EText3DRendererFlags InFlags)
 {
+	using namespace UE::Text3D::Renderer;
+
 	if (!bInitialized)
 	{
 		return;
@@ -81,13 +83,19 @@ void UText3DRendererBase::Update(EText3DRendererFlags InFlags)
 	{
 		const EText3DRendererFlags FlagCasted = static_cast<EText3DRendererFlags>(Flag);
 
+		const FUpdateParameters Parameters
+		{
+			.UpdateFlags = InFlags,
+			.CurrentFlag = FlagCasted
+		};
+
 		if (EnumHasAnyFlags(InFlags, FlagCasted))
 		{
 			for (int32 Index = Extensions.Num() - 1; Index >= 0; --Index)
 			{
 				if (UText3DExtensionBase* Extension = Extensions[Index])
 				{
-					const EText3DExtensionResult ExtensionStatus = Extension->PreRendererUpdate(FlagCasted);
+					const EText3DExtensionResult ExtensionStatus = Extension->PreRendererUpdate(Parameters);
 
 					if (ExtensionStatus == EText3DExtensionResult::Failed)
 					{
@@ -101,13 +109,13 @@ void UText3DRendererBase::Update(EText3DRendererFlags InFlags)
 				}
 			}
 
-			OnUpdate(FlagCasted);
+			OnUpdate(Parameters);
 
 			for (int32 Index = Extensions.Num() - 1; Index >= 0; --Index)
 			{
 				if (UText3DExtensionBase* Extension = Extensions[Index])
 				{
-					const EText3DExtensionResult ExtensionStatus = Extension->PostRendererUpdate(FlagCasted);
+					const EText3DExtensionResult ExtensionStatus = Extension->PostRendererUpdate(Parameters);
 
 					if (ExtensionStatus == EText3DExtensionResult::Failed)
 					{
@@ -125,7 +133,7 @@ void UText3DRendererBase::Update(EText3DRendererFlags InFlags)
 		Flag <<= 1;
 	}
 
-	UE_LOG(LogText3D, Verbose, TEXT("%s : Text3DRenderer %s Updated with flags %i"), *Text3DComponent->GetOwner()->GetActorNameOrLabel(), *GetName().ToString(), InFlags)
+	UE_LOG(LogText3D, Verbose, TEXT("%s : Text3DRenderer %s Updated with flags %i"), *Text3DComponent->GetOwner()->GetActorNameOrLabel(), *GetFriendlyName().ToString(), InFlags)
 }
 
 void UText3DRendererBase::Clear()
@@ -145,7 +153,7 @@ void UText3DRendererBase::Clear()
 	CachedBounds.Reset();
 
 	const AActor* Owner = Text3DComponent->GetOwner();
-	UE_LOG(LogText3D, Verbose, TEXT("%s : Text3DRenderer %s Cleared"), !!Owner ? *Owner->GetActorNameOrLabel() : TEXT("Invalid owner"), *GetName().ToString())
+	UE_LOG(LogText3D, Verbose, TEXT("%s : Text3DRenderer %s Cleared"), !!Owner ? *Owner->GetActorNameOrLabel() : TEXT("Invalid owner"), *GetFriendlyName().ToString())
 }
 
 void UText3DRendererBase::Destroy()
@@ -166,7 +174,7 @@ void UText3DRendererBase::Destroy()
 	bInitialized = false;
 
 	const AActor* Owner = Text3DComponent->GetOwner();
-	UE_LOG(LogText3D, Log, TEXT("%s : Text3DRenderer %s Destroyed"), !!Owner ? *Owner->GetActorNameOrLabel() : TEXT("Invalid owner"), *GetName().ToString())
+	UE_LOG(LogText3D, Log, TEXT("%s : Text3DRenderer %s Destroyed"), !!Owner ? *Owner->GetActorNameOrLabel() : TEXT("Invalid owner"), *GetFriendlyName().ToString())
 }
 
 FBox UText3DRendererBase::GetBounds() const
diff --git a/Engine/Plugins/Experimental/Text3D/Source/Text3D/Private/Subsystems/Text3DEngineSubsystem.cpp b/Engine/Plugins/Experimental/Text3D/Source/Text3D/Private/Subsystems/Text3DEngineSubsystem.cpp
--- a/Engine/Plugins/Experimental/Text3D/Source/Text3D/Private/Subsystems/Text3DEngineSubsystem.cpp
+++ b/Engine/Plugins/Experimental/Text3D/Source/Text3D/Private/Subsystems/Text3DEngineSubsystem.cpp
@@ -327,7 +327,7 @@ FText3DCachedMesh* FCachedFontData::GetGlyphMesh(uint32 InGlyphIndex, const FGly
 	const FName StrokeShapeUniqueName = MakeUniqueObjectName(Subsystem, UStaticMesh::StaticClass(), FName(*StaticMeshName));
 
 	FText3DCachedMesh& CachedMesh = CachedMeshes.Glyphs.FindOrAdd(InGlyphIndex);
-	CachedMesh.Mesh = NewObject<UStaticMesh>(Subsystem, StrokeShapeUniqueName);
+	CachedMesh.Mesh = NewObject<UStaticMesh>(Subsystem, StrokeShapeUniqueName, RF_Transient | RF_Public);
 
 	FText3DGlyphMeshBuilder MeshCreator;
 	MeshCreator.CreateMeshes(Root, InParameters.Extrude, InParameters.Bevel, InParameters.BevelType, InParameters.BevelSegments, InParameters.bOutline, InParameters.OutlineExpand);
diff --git a/Engine/Plugins/Experimental/Text3D/Source/Text3D/Private/Text3DComponent.cpp b/Engine/Plugins/Experimental/Text3D/Source/Text3D/Private/Text3DComponent.cpp
--- a/Engine/Plugins/Experimental/Text3D/Source/Text3D/Private/Text3DComponent.cpp
+++ b/Engine/Plugins/Experimental/Text3D/Source/Text3D/Private/Text3DComponent.cpp
@@ -19,17 +19,14 @@
 #include "Misc/EnumerateRange.h"
 #include "Misc/ScopeExit.h"
 #include "Misc/TransactionObjectEvent.h"
-#include "Renderers/Text3DRendererBase.h"
 #include "Renderers/StaticMeshes/Text3DStaticMeshesRenderer.h"
+#include "Renderers/Text3DRendererBase.h"
 #include "Settings/Text3DProjectSettings.h"
 #include "Styling/StyleDefaults.h"
 #include "Subsystems/Text3DEngineSubsystem.h"
-#include "Text3DModule.h"
 #include "Text3DComponentVersion.h"
-#include "Text3DModule.h"
 #include "UObject/ConstructorHelpers.h"
 #include "UObject/UObjectGlobals.h"
-#include "Utilities/Text3DUtilities.h"
 
 #define LOCTEXT_NAMESPACE "Text3D"
 
@@ -91,11 +88,6 @@ void UText3DComponent::PostLoad()
 	OnTextRendererClassChanged();
 }
 
-EText3DRendererFlags UText3DComponent::GetUpdateFlags() const
-{
-	return UpdateFlags;
-}
-
 void UText3DComponent::RequestUpdate(EText3DRendererFlags InFlags, bool bInImmediate)
 {
 	if (InFlags != EText3DRendererFlags::None)
@@ -233,14 +225,10 @@ void UText3DComponent::UpdateStatistics()
 
 void UText3DComponent::OnComponentTransformed(USceneComponent*, EUpdateTransformFlags, ETeleportType)
 {
-	static TMap<FObjectKey, bool> TextComponentScale;
-
 	/*
-	 * Find previous zero scaled status or initialize it to false
 	 * TransformUpdated gets broadcasted when component is registered first time so it will have the correct starting value (see UpdateComponentToWorld in UActorComponent::OnRegister)
 	 * Further calls affecting transform value, will have the proper zero scaled status
 	 */
-	bool& bZeroScaled = TextComponentScale.FindOrAdd(this, /** Default */false);
 	const bool bIsZeroScaled = FMath::IsNearlyZero(GetComponentScale().GetMin(), UE_KINDA_SMALL_NUMBER);
 
 	if (bIsZeroScaled)
@@ -900,18 +888,10 @@ void UText3DComponent::SetTextRendererClass(const TSubclassOf<UText3DRendererBas
 	OnTextRendererClassChanged();
 }
 
-void UText3DComponent::RebuildInternal(bool bCleanCache)
-{
-}
-
 void UText3DComponent::QueueRebuild(EText3DRendererFlags InFlags)
 {
-	// Hotfix: To not break ABI compatibility, keep track of the queued flags here
-	static TMap<FObjectKey, EText3DRendererFlags> QueuedUpdateFlags;
-
 	// Add the given flags to the current queued flags in case a text update is already queued
-	EText3DRendererFlags& CurrentUpdateFlags = QueuedUpdateFlags.FindOrAdd(this, EText3DRendererFlags::None);
-	EnumAddFlags(CurrentUpdateFlags, InFlags);
+	EnumAddFlags(UpdateFlags, InFlags);
 
 	if (!TextUpdateHandle.IsValid())
 	{
@@ -930,10 +910,10 @@ void UText3DComponent::QueueRebuild(EText3DRendererFlags InFlags)
 					// New rebuild requests should only happen when, while building text, existing 'built data' becomes outdated.
 					TextUpdateHandle.Reset();
 
-					EText3DRendererFlags LatestUpdateFlags = EText3DRendererFlags::None;
-					QueuedUpdateFlags.RemoveAndCopyValue(this, LatestUpdateFlags);
+					const EText3DRendererFlags CurrentUpdateFlags = UpdateFlags;
+					UpdateFlags = EText3DRendererFlags::None;
 
-					Rebuild(LatestUpdateFlags, /*bCleanCache*/false);
+					Rebuild(CurrentUpdateFlags, /*bCleanCache*/false);
 
 					// Executes only once
 					return false;
@@ -945,9 +925,7 @@ void UText3DComponent::QueueRebuild(EText3DRendererFlags InFlags)
 
 void UText3DComponent::Rebuild(EText3DRendererFlags InUpdateFlags, bool bCleanCache)
 {
-	EnumAddFlags(UpdateFlags, InUpdateFlags);
-
-	if (UpdateFlags == EText3DRendererFlags::None)
+	if (InUpdateFlags == EText3DRendererFlags::None)
 	{
 		return;
 	}
@@ -999,26 +977,23 @@ void UText3DComponent::Rebuild(EText3DRendererFlags InUpdateFlags, bool bCleanCa
 	ON_SCOPE_EXIT
 	{
 		bIsUpdatingText = false;
-		EnumRemoveFlags(UpdateFlags, InUpdateFlags);
 	};
 
-	TextPreUpdateDelegate.Broadcast(this, UpdateFlags);
+	TextPreUpdateDelegate.Broadcast(this, InUpdateFlags);
 
 	TRACE_CPUPROFILER_EVENT_SCOPE(UText3DComponent::RebuildInternal);
 
 	bool bBroadcastEvent = false;
 
-	if (EnumHasAnyFlags(UpdateFlags,EText3DRendererFlags::Geometry))
+	if (EnumHasAnyFlags(InUpdateFlags,EText3DRendererFlags::Geometry))
 	{
 		FormattedText = Text;
 		FormatText(FormattedText.GetValue());
 
-		TextRenderer->Clear();
-
 		bBroadcastEvent = true;
 	}
 
-	TextRenderer->Update(UpdateFlags);
+	TextRenderer->Update(InUpdateFlags);
 
 	if (bBroadcastEvent)
 	{
@@ -1028,7 +1003,7 @@ void UText3DComponent::Rebuild(EText3DRendererFlags InUpdateFlags, bool bCleanCa
 		TextGeneratedDelegate.Broadcast();
 	}
 
-	if (EnumHasAnyFlags(UpdateFlags,EText3DRendererFlags::Visibility))
+	if (EnumHasAnyFlags(InUpdateFlags,EText3DRendererFlags::Visibility))
 	{
 		MarkRenderStateDirty();
 	}
@@ -1041,11 +1016,7 @@ void UText3DComponent::Rebuild(EText3DRendererFlags InUpdateFlags, bool bCleanCa
 		}
 	}
 
-	TextPostUpdateDelegate.Broadcast(this, UpdateFlags);
-}
-
-void UText3DComponent::ClearUpdateFlags()
-{
+	TextPostUpdateDelegate.Broadcast(this, InUpdateFlags);
 }
 
 void UText3DComponent::OnTextRendererClassChanged()
@@ -1071,7 +1042,7 @@ void UText3DComponent::OnTextRendererClassChanged()
 			return;
 		}
 
-		OldRendererName = TextRenderer->GetName();
+		OldRendererName = TextRenderer->GetFriendlyName();
 
 		TextRenderer->Clear();
 		TextRenderer->Destroy();
@@ -1079,10 +1050,10 @@ void UText3DComponent::OnTextRendererClassChanged()
 		TextRenderer = nullptr;
 	}
 
-	TextRenderer = NewObject<UText3DRendererBase>(this, TextRendererClass.Get(), RendererCDO->GetName(), RF_Transient);
+	TextRenderer = NewObject<UText3DRendererBase>(this, TextRendererClass.Get(), RendererCDO->GetFriendlyName(), RF_Public);
 
 	const AActor* ContextActor = GetOwner();
-	UE_LOG(LogText3D, Log, TEXT("%s : Text3D renderer changed : old %s - new %s"), !!ContextActor ? *ContextActor->GetActorNameOrLabel() : TEXT("Invalid owner"), *OldRendererName.ToString(), *TextRenderer->GetName().ToString());
+	UE_LOG(LogText3D, Log, TEXT("%s : Text3D renderer changed : old %s - new %s"), !!ContextActor ? *ContextActor->GetActorNameOrLabel() : TEXT("Invalid owner"), *OldRendererName.ToString(), *TextRenderer->GetFriendlyName().ToString());
 
 	if (IsRegistered())
 	{
@@ -1105,10 +1076,6 @@ void UText3DComponent::OnFontPropertiesChanged()
 	RequestUpdate(EText3DRendererFlags::All);
 }
 
-void UText3DComponent::ScheduleTextUpdateNextTick()
-{
-}
-
 TArray<FName> UText3DComponent::GetTypefaceNames() const
 {
 	TArray<FName> TypefaceNames;
diff --git a/Engine/Plugins/Experimental/Text3D/Source/Text3D/Public/Extensions/Text3DDefaultGeometryExtension.h b/Engine/Plugins/Experimental/Text3D/Source/Text3D/Public/Extensions/Text3DDefaultGeometryExtension.h
--- a/Engine/Plugins/Experimental/Text3D/Source/Text3D/Public/Extensions/Text3DDefaultGeometryExtension.h
+++ b/Engine/Plugins/Experimental/Text3D/Source/Text3D/Public/Extensions/Text3DDefaultGeometryExtension.h
@@ -109,8 +109,8 @@ protected:
 	//~ End UText3DGeometryExtensionBase
 
 	//~ Begin UText3DExtensionBase
-	virtual EText3DExtensionResult PreRendererUpdate(EText3DRendererFlags InFlag) override;
-	virtual EText3DExtensionResult PostRendererUpdate(EText3DRendererFlags InFlag) override;
+	virtual EText3DExtensionResult PreRendererUpdate(const UE::Text3D::Renderer::FUpdateParameters& InParameters) override;
+	virtual EText3DExtensionResult PostRendererUpdate(const UE::Text3D::Renderer::FUpdateParameters& InParameters) override;
 	//~ End UText3DExtensionBase
 
 	void OnGeometryOptionsChanged();
diff --git a/Engine/Plugins/Experimental/Text3D/Source/Text3D/Public/Extensions/Text3DDefaultLayoutExtension.h b/Engine/Plugins/Experimental/Text3D/Source/Text3D/Public/Extensions/Text3DDefaultLayoutExtension.h
--- a/Engine/Plugins/Experimental/Text3D/Source/Text3D/Public/Extensions/Text3DDefaultLayoutExtension.h
+++ b/Engine/Plugins/Experimental/Text3D/Source/Text3D/Public/Extensions/Text3DDefaultLayoutExtension.h
@@ -158,8 +158,8 @@ protected:
 	//~ End UText3DLayoutExtensionBase
 
 	//~ Begin UText3DExtensionBase
-	virtual EText3DExtensionResult PreRendererUpdate(EText3DRendererFlags InFlag) override;
-	virtual EText3DExtensionResult PostRendererUpdate(EText3DRendererFlags InFlag) override;
+	virtual EText3DExtensionResult PreRendererUpdate(const UE::Text3D::Renderer::FUpdateParameters& InParameters) override;
+	virtual EText3DExtensionResult PostRendererUpdate(const UE::Text3D::Renderer::FUpdateParameters& InParameters) override;
 	//~ End UText3DExtensionBase
 
 	void CalculateTextScale();
diff --git a/Engine/Plugins/Experimental/Text3D/Source/Text3D/Public/Extensions/Text3DDefaultMaterialExtension.h b/Engine/Plugins/Experimental/Text3D/Source/Text3D/Public/Extensions/Text3DDefaultMaterialExtension.h
--- a/Engine/Plugins/Experimental/Text3D/Source/Text3D/Public/Extensions/Text3DDefaultMaterialExtension.h
+++ b/Engine/Plugins/Experimental/Text3D/Source/Text3D/Public/Extensions/Text3DDefaultMaterialExtension.h
@@ -244,8 +244,8 @@ protected:
 	//~ End UObject
 
 	//~ Begin UText3DExtensionBase
-	virtual EText3DExtensionResult PreRendererUpdate(EText3DRendererFlags InFlag) override;
-	virtual EText3DExtensionResult PostRendererUpdate(EText3DRendererFlags InFlag) override;
+	virtual EText3DExtensionResult PreRendererUpdate(const UE::Text3D::Renderer::FUpdateParameters& InParameters) override;
+	virtual EText3DExtensionResult PostRendererUpdate(const UE::Text3D::Renderer::FUpdateParameters& InParameters) override;
 	//~ End UText3DExtensionBase
 
 	//~ Begin UText3DMaterialExtensionBase
diff --git a/Engine/Plugins/Experimental/Text3D/Source/Text3D/Public/Extensions/Text3DEffectExtensionBase.h b/Engine/Plugins/Experimental/Text3D/Source/Text3D/Public/Extensions/Text3DEffectExtensionBase.h
--- a/Engine/Plugins/Experimental/Text3D/Source/Text3D/Public/Extensions/Text3DEffectExtensionBase.h
+++ b/Engine/Plugins/Experimental/Text3D/Source/Text3D/Public/Extensions/Text3DEffectExtensionBase.h
@@ -61,8 +61,8 @@ public:
 
 protected:
 	//~ Begin UText3DExtensionBase
-	virtual EText3DExtensionResult PreRendererUpdate(EText3DRendererFlags InFlag) override;
-	virtual EText3DExtensionResult PostRendererUpdate(EText3DRendererFlags InFlag) override;
+	virtual EText3DExtensionResult PreRendererUpdate(const UE::Text3D::Renderer::FUpdateParameters& InParameters) override;
+	virtual EText3DExtensionResult PostRendererUpdate(const UE::Text3D::Renderer::FUpdateParameters& InParameters) override;
 	//~ End UText3DExtensionBase
 
 	virtual FText3DRange GetTargetRange() const
diff --git a/Engine/Plugins/Experimental/Text3D/Source/Text3D/Public/Extensions/Text3DExtensionBase.h b/Engine/Plugins/Experimental/Text3D/Source/Text3D/Public/Extensions/Text3DExtensionBase.h
--- a/Engine/Plugins/Experimental/Text3D/Source/Text3D/Public/Extensions/Text3DExtensionBase.h
+++ b/Engine/Plugins/Experimental/Text3D/Source/Text3D/Public/Extensions/Text3DExtensionBase.h
@@ -36,13 +36,13 @@ public:
 	{}
 
 	/** Perform an update of the extension behavior before renderer execute based on condition above */
-	virtual EText3DExtensionResult PreRendererUpdate(EText3DRendererFlags InFlag)
+	virtual EText3DExtensionResult PreRendererUpdate(const UE::Text3D::Renderer::FUpdateParameters& InParameters)
 	{
 		return EText3DExtensionResult::Finished;
 	}
 
 	/** Perform an update of the extension behavior after renderer execute based on condition above */
-	virtual EText3DExtensionResult PostRendererUpdate(EText3DRendererFlags InFlag)
+	virtual EText3DExtensionResult PostRendererUpdate(const UE::Text3D::Renderer::FUpdateParameters& InParameters)
 	{
 		return EText3DExtensionResult::Finished;
 	}
diff --git a/Engine/Plugins/Experimental/Text3D/Source/Text3D/Public/Renderers/Text3DRendererBase.h b/Engine/Plugins/Experimental/Text3D/Source/Text3D/Public/Renderers/Text3DRendererBase.h
--- a/Engine/Plugins/Experimental/Text3D/Source/Text3D/Public/Renderers/Text3DRendererBase.h
+++ b/Engine/Plugins/Experimental/Text3D/Source/Text3D/Public/Renderers/Text3DRendererBase.h
@@ -34,14 +34,14 @@ public:
 	FBox GetBounds() const;
 
 	/** Get the implementation name for debug purposes */
-	virtual FName GetName() const PURE_VIRTUAL(UText3DRenderingImplementationBase::GetName, return NAME_None; );
+	virtual FName GetFriendlyName() const PURE_VIRTUAL(UText3DRenderingImplementationBase::GetFriendlyName, return NAME_None; );
 	
 protected:
 	/** Create and setup the implementation components, called on load or creation */
 	virtual void OnCreate() PURE_VIRTUAL(UText3DRenderingImplementationBase::OnCreate, );
 
 	/** Update rendering state of text characters, called when render state is outdated */
-	virtual void OnUpdate(EText3DRendererFlags InFlags) PURE_VIRTUAL(UText3DRenderingImplementationBase::OnUpdate, );
+	virtual void OnUpdate(const UE::Text3D::Renderer::FUpdateParameters& InParameters) PURE_VIRTUAL(UText3DRenderingImplementationBase::OnUpdate, );
 
 	/** Clear rendering state and remove all visible characters */
 	virtual void OnClear() PURE_VIRTUAL(UText3DRenderingImplementationBase::OnClear, );
diff --git a/Engine/Plugins/Experimental/Text3D/Source/Text3D/Public/Text3DComponent.h b/Engine/Plugins/Experimental/Text3D/Source/Text3D/Public/Text3DComponent.h
--- a/Engine/Plugins/Experimental/Text3D/Source/Text3D/Public/Text3DComponent.h
+++ b/Engine/Plugins/Experimental/Text3D/Source/Text3D/Public/Text3DComponent.h
@@ -260,15 +260,13 @@ public:
 		return Statistics;
 	}
 
-	EText3DRendererFlags GetUpdateFlags() const;
-
 	TEXT3D_API void RequestUpdate(EText3DRendererFlags InFlags, bool bInImmediate = false);
 
 	uint32 GetTypeFaceIndex() const;
 
 	const FTypefaceEntry* GetTypeFaceEntry() const;
 	
-	UText3DRendererBase* GetTextRenderer() const;
+	TEXT3D_API UText3DRendererBase* GetTextRenderer() const;
 
 	UFUNCTION(BlueprintCallable, Category="Text3D", meta=(DeterminesOutputType="InExtensionClass"))
 	UText3DLayoutExtensionBase* GetLayoutExtension(TSubclassOf<UText3DLayoutExtensionBase> InExtensionClass);
@@ -422,9 +420,6 @@ protected:
 	/** Called when font options are changed */
 	void OnFontPropertiesChanged();
 
-	/** Todo: Remove for 5.7. Kept here for ABI Compatibility */
-	void ScheduleTextUpdateNextTick();
-
 	/** The renderer class to use to create the text geometry */
 	UPROPERTY(EditAnywhere, Setter, Getter, Category="Text", NoClear, meta=(ShowDisplayNames, AllowPrivateAccess="true"))
 	TSubclassOf<UText3DRendererBase> TextRendererClass;
@@ -469,10 +464,10 @@ protected:
 
 private:
 	/** Active renderer in charge of generating text geometry */
-	UPROPERTY()
+	UPROPERTY(VisibleAnywhere, Instanced, NoClear, Category="Text", meta=(AllowPrivateAccess="true", NoResetToDefault))
 	TObjectPtr<UText3DRendererBase> TextRenderer;
 
-	UPROPERTY(BlueprintAssignable, Category = Events, meta = (AllowPrivateAccess = "true", DisplayName = "On Text Generated"))
+	UPROPERTY(BlueprintAssignable, Category=Events, meta=(AllowPrivateAccess="true", DisplayName="On Text Generated"))
 	FTextGenerated TextGeneratedDelegate;
 
 	FTextGeneratedNative TextGeneratedNativeDelegate;
@@ -495,18 +490,12 @@ private:
 	/** Planned update requested */
 	FTSTicker::FDelegateHandle TextUpdateHandle;
 
-	/** Todo: Remove for 5.7. Kept here for ABI Compatibility */
-	void RebuildInternal(bool bCleanCache = false);
-
 	/** Queues text rebuild for the next tick */
 	void QueueRebuild(EText3DRendererFlags InFlags);
 
 	/** Partially or fully rebuilds text depending on the given update flags */
 	void Rebuild(EText3DRendererFlags InUpdateFlags, bool bCleanCache);
 
-	/** Todo: Remove for 5.7. Kept here for ABI Compatibility */
-	void ClearUpdateFlags();
-
 	bool IsTypefaceAvailable(FName InTypeface) const;
 	TArray<FTypefaceEntry> GetAvailableTypefaces() const;
 	void RefreshTypeface();
@@ -515,6 +504,8 @@ private:
 
 	void OnComponentTransformed(USceneComponent*, EUpdateTransformFlags, ETeleportType);
 
+	bool bZeroScaled = false;
+
 #if WITH_EDITORONLY_DATA
 	UE_DEPRECATED(5.6, "Use Material Extension instead")
 	UPROPERTY()
diff --git a/Engine/Plugins/Experimental/Text3D/Source/Text3D/Public/Text3DTypes.h b/Engine/Plugins/Experimental/Text3D/Source/Text3D/Public/Text3DTypes.h
--- a/Engine/Plugins/Experimental/Text3D/Source/Text3D/Public/Text3DTypes.h
+++ b/Engine/Plugins/Experimental/Text3D/Source/Text3D/Public/Text3DTypes.h
@@ -221,14 +221,29 @@ struct FText3DTypeFaceMetrics
 	float FontDescender = 0.f;
 };
 
-namespace UE::Text3D::Metrics
-{
-	/** Original value to match size across different Text3D version */
-	static constexpr float FontSize = 64;
-	/** Value used to match size across Text3D implementation */
-	static constexpr float SlateFontSize = 48;
-	/** Scale used to transform freetype face result to get normalized values */
-	static constexpr float FontSizeInverse = 1.0f / FontSize;
-	/** DPI used to match slate */
-	static constexpr float FontDPI = 96.f;
+namespace UE::Text3D
+{
+	namespace Metrics
+	{
+		/** Original value to match size across different Text3D version */
+		static constexpr float FontSize = 64;
+		/** Value used to match size across Text3D implementation */
+		static constexpr float SlateFontSize = 48;
+		/** Scale used to transform freetype face result to get normalized values */
+		static constexpr float FontSizeInverse = 1.0f / FontSize;
+		/** DPI used to match slate */
+		static constexpr float FontDPI = 96.f;
+	}
+
+	namespace Renderer
+	{
+		struct FUpdateParameters
+		{
+			/** All dirty flags for this update */
+			EText3DRendererFlags UpdateFlags = EText3DRendererFlags::None;
+
+			/** Current flag being updated */
+			EText3DRendererFlags CurrentFlag = EText3DRendererFlags::None;
+		};
+	}
 }
\ No newline at end of file
