MotionDesign : Text3D
- Fix robomerge conflict coming from 5.6.1 to 5.7

diff --git a/Engine/Plugins/Experimental/Text3D/Source/Text3D/Private/Renderers/StaticMeshes/Text3DStaticMeshesRenderer.cpp b/Engine/Plugins/Experimental/Text3D/Source/Text3D/Private/Renderers/StaticMeshes/Text3DStaticMeshesRenderer.cpp
--- a/Engine/Plugins/Experimental/Text3D/Source/Text3D/Private/Renderers/StaticMeshes/Text3DStaticMeshesRenderer.cpp
+++ b/Engine/Plugins/Experimental/Text3D/Source/Text3D/Private/Renderers/StaticMeshes/Text3DStaticMeshesRenderer.cpp
@@ -11,6 +11,7 @@
 #include "Extensions/Text3DLayoutExtensionBase.h"
 #include "Extensions/Text3DMaterialExtensionBase.h"
 #include "Extensions/Text3DRenderingExtensionBase.h"
+#include "Logs/Text3DLogs.h"
 #include "Text3DComponent.h"
 #include "Text3DInternalTypes.h"
 
@@ -37,7 +38,14 @@ void UText3DStaticMeshesRenderer::OnUpdate(const UE::Text3D::Renderer::FUpdatePa
 		TextComponent->ForEachCharacter([this, GeometryExtension](UText3DCharacterBase* InCharacter, uint16 InIndex, uint16 InCount)
 		{
 			UStaticMesh* CachedMesh = GeometryExtension->FindOrLoadGlyphMesh(InCharacter);
-			CharacterMeshes[InIndex]->SetStaticMesh(CachedMesh);
+			if (CharacterMeshes.IsValidIndex(InIndex))
+			{
+				CharacterMeshes[InIndex]->SetStaticMesh(CachedMesh);
+			}
+			else
+			{
+				UE_LOG(LogText3D, Warning, TEXT("Invalid component index for character %i"), InIndex)
+			}
 		});
 
 		RefreshBounds();
@@ -52,7 +60,15 @@ void UText3DStaticMeshesRenderer::OnUpdate(const UE::Text3D::Renderer::FUpdatePa
 		{
 			constexpr bool bReset = false;
 			const FTransform& CharacterTransform = InCharacter->GetTransform(bReset);
-			CharacterMeshes[InIndex]->SetRelativeTransform(CharacterTransform);
+
+			if (CharacterMeshes.IsValidIndex(InIndex))
+			{
+				CharacterMeshes[InIndex]->SetRelativeTransform(CharacterTransform);
+			}
+			else
+			{
+				UE_LOG(LogText3D, Warning, TEXT("Invalid component index for character %i"), InIndex)
+			}
 		});
 
 		RefreshBounds();
@@ -65,23 +81,30 @@ void UText3DStaticMeshesRenderer::OnUpdate(const UE::Text3D::Renderer::FUpdatePa
 
 		TextComponent->ForEachCharacter([this, MaterialExtension](UText3DCharacterBase* InCharacter, uint16 InIndex, uint16 InCount)
 		{
-			UStaticMeshComponent* StaticMeshComponent = CharacterMeshes[InIndex];
-			for (int32 GroupIndex = 0; GroupIndex < static_cast<int32>(EText3DGroupType::TypeCount); GroupIndex++)
+			if (CharacterMeshes.IsValidIndex(InIndex))
 			{
-				const int32 MaterialIndex = StaticMeshComponent->GetMaterialIndex(SlotNames[GroupIndex]);
-
-				if (MaterialIndex == INDEX_NONE)
+				UStaticMeshComponent* StaticMeshComponent = CharacterMeshes[InIndex];
+				for (int32 GroupIndex = 0; GroupIndex < static_cast<int32>(EText3DGroupType::TypeCount); GroupIndex++)
 				{
-					continue;
-				}
+					const int32 MaterialIndex = StaticMeshComponent->GetMaterialIndex(SlotNames[GroupIndex]);
 
-				UMaterialInterface* Material = MaterialExtension->GetMaterial(static_cast<EText3DGroupType>(GroupIndex));
+					if (MaterialIndex == INDEX_NONE)
+					{
+						continue;
+					}
 
-				if (Material != StaticMeshComponent->GetMaterial(MaterialIndex))
-				{
-					StaticMeshComponent->SetMaterial(MaterialIndex, Material);
+					UMaterialInterface* Material = MaterialExtension->GetMaterial(static_cast<EText3DGroupType>(GroupIndex));
+
+					if (Material != StaticMeshComponent->GetMaterial(MaterialIndex))
+					{
+						StaticMeshComponent->SetMaterial(MaterialIndex, Material);
+					}
 				}
 			}
+			else
+			{
+				UE_LOG(LogText3D, Warning, TEXT("Invalid component index for character %i"), InIndex)
+			}
 		});
 	}
 	else if (InParameters.CurrentFlag == EText3DRendererFlags::Visibility)
@@ -90,14 +113,21 @@ void UText3DStaticMeshesRenderer::OnUpdate(const UE::Text3D::Renderer::FUpdatePa
 
 		TextComponent->ForEachCharacter([this, TextComponent, RenderingExtension](UText3DCharacterBase* InCharacter, uint16 InIndex, uint16 InCount)
 		{
-			UStaticMeshComponent* StaticMeshComponent = CharacterMeshes[InIndex];
-			StaticMeshComponent->SetHiddenInGame(TextComponent->bHiddenInGame);
-			StaticMeshComponent->SetVisibility(TextComponent->GetVisibleFlag() && InCharacter->GetVisibility());
-			StaticMeshComponent->SetCastShadow(RenderingExtension->GetTextCastShadow());
-			StaticMeshComponent->SetCastHiddenShadow(RenderingExtension->GetTextCastHiddenShadow());
-			StaticMeshComponent->SetAffectDynamicIndirectLighting(RenderingExtension->GetTextAffectDynamicIndirectLighting());
-			StaticMeshComponent->SetAffectIndirectLightingWhileHidden(RenderingExtension->GetTextAffectIndirectLightingWhileHidden());
-			StaticMeshComponent->SetHoldout(RenderingExtension->GetTextHoldout());
+			if (CharacterMeshes.IsValidIndex(InIndex))
+			{
+				UStaticMeshComponent* StaticMeshComponent = CharacterMeshes[InIndex];
+				StaticMeshComponent->SetHiddenInGame(TextComponent->bHiddenInGame);
+				StaticMeshComponent->SetVisibility(TextComponent->GetVisibleFlag() && InCharacter->GetVisibility());
+				StaticMeshComponent->SetCastShadow(RenderingExtension->GetTextCastShadow());
+				StaticMeshComponent->SetCastHiddenShadow(RenderingExtension->GetTextCastHiddenShadow());
+				StaticMeshComponent->SetAffectDynamicIndirectLighting(RenderingExtension->GetTextAffectDynamicIndirectLighting());
+				StaticMeshComponent->SetAffectIndirectLightingWhileHidden(RenderingExtension->GetTextAffectIndirectLightingWhileHidden());
+				StaticMeshComponent->SetHoldout(RenderingExtension->GetTextHoldout());
+			}
+			else
+			{
+				UE_LOG(LogText3D, Warning, TEXT("Invalid component index for character %i"), InIndex)
+			}
 		});
 	}
 }
