MotionDesign : Text3D
- Export text to merged static mesh asset
- Expose ufunction from text component and renderer in details panel
- Fixed text duplicated components in PIE

diff --git a/Engine/Plugins/Experimental/Text3D/Source/Text3D/Private/Renderers/StaticMeshes/Text3DStaticMeshesRenderer.cpp b/Engine/Plugins/Experimental/Text3D/Source/Text3D/Private/Renderers/StaticMeshes/Text3DStaticMeshesRenderer.cpp
--- a/Engine/Plugins/Experimental/Text3D/Source/Text3D/Private/Renderers/StaticMeshes/Text3DStaticMeshesRenderer.cpp
+++ b/Engine/Plugins/Experimental/Text3D/Source/Text3D/Private/Renderers/StaticMeshes/Text3DStaticMeshesRenderer.cpp
@@ -14,6 +14,13 @@
 #include "Logs/Text3DLogs.h"
 #include "Text3DComponent.h"
 #include "Text3DInternalTypes.h"
+#include "Utilities/Text3DUtilities.h"
+
+#if WITH_EDITOR
+#include "Misc/ScopedSlowTask.h"
+#endif
+
+#define LOCTEXT_NAMESPACE "Text3DStaticMeshesRenderer"
 
 void UText3DStaticMeshesRenderer::OnCreate()
 {
@@ -226,8 +233,9 @@ void UText3DStaticMeshesRenderer::AllocateComponents(int32 InCount)
 				CharacterMeshes.Add(MeshComponent);
 			}
 
-			MeshComponent->AttachToComponent(TextRoot, FAttachmentTransformRules::KeepRelativeTransform);
-			MeshComponent->RegisterComponent();
+			MeshComponent->SetupAttachment(TextRoot);
+			Actor->AddInstanceComponent(MeshComponent);
+			Actor->AddOwnedComponent(MeshComponent);
 		}
 	}
 
@@ -258,3 +266,36 @@ const TArray<UStaticMeshComponent*>& UText3DStaticMeshesRenderer::GetGlyphMeshCo
 {
 	return CharacterMeshes;
 }
+
+#if WITH_EDITOR
+void UText3DStaticMeshesRenderer::ConvertToStaticMesh()
+{
+	if (!IsValid(this))
+	{
+		return;
+	}
+
+	const UText3DComponent* Text3DComponent = GetText3DComponent();
+	if (!IsValid(Text3DComponent) || !IsValid(Text3DComponent->GetOwner()))
+	{
+		return;
+	}
+
+	const AActor* Owner = Text3DComponent->GetOwner();
+	FScopedSlowTask SlowTask(0.0f, LOCTEXT("ConvertToStaticMesh", "Converting Text3D to static mesh"));
+	SlowTask.MakeDialog();
+
+	UE_LOG(LogText3D, Log, TEXT("%s : Request ConvertToStaticMesh..."), *Owner->GetActorNameOrLabel())
+
+	if (UE::Text3D::Utilities::Conversion::ConvertToStaticMesh(Text3DComponent))
+	{
+		UE_LOG(LogText3D, Log, TEXT("%s : ConvertToStaticMesh Completed"), *Owner->GetActorNameOrLabel())
+	}
+	else
+	{
+		UE_LOG(LogText3D, Warning, TEXT("%s : ConvertToStaticMesh Failed"), *Owner->GetActorNameOrLabel())
+	}
+}
+#endif
+
+#undef LOCTEXT_NAMESPACE 
\ No newline at end of file
diff --git a/Engine/Plugins/Experimental/Text3D/Source/Text3D/Private/Renderers/StaticMeshes/Text3DStaticMeshesRenderer.h b/Engine/Plugins/Experimental/Text3D/Source/Text3D/Private/Renderers/StaticMeshes/Text3DStaticMeshesRenderer.h
--- a/Engine/Plugins/Experimental/Text3D/Source/Text3D/Private/Renderers/StaticMeshes/Text3DStaticMeshesRenderer.h
+++ b/Engine/Plugins/Experimental/Text3D/Source/Text3D/Private/Renderers/StaticMeshes/Text3DStaticMeshesRenderer.h
@@ -39,6 +39,12 @@ public:
 	UFUNCTION(BlueprintCallable, Category = "Rendering|Components|Text3D")
 	const TArray<UStaticMeshComponent*>& GetGlyphMeshComponents();
 
+#if WITH_EDITOR
+	/** Converts the text to a single static mesh, this is an heavy operation */
+	UFUNCTION(CallInEditor, Category="Utilities")
+	void ConvertToStaticMesh();
+#endif
+
 protected:
 	//~ Begin UText3DRendererBase
 	virtual void OnCreate() override;
@@ -52,11 +58,11 @@ protected:
 	void AllocateComponents(int32 InCount);
 
 	/** Holds all character components */
-	UPROPERTY(VisibleAnywhere, DuplicateTransient, TextExportTransient, Category="Renderer", meta=(EditCondition="false", EditConditionHides))
+	UPROPERTY(VisibleAnywhere, Category="Renderer", meta=(EditCondition="false", EditConditionHides))
 	TObjectPtr<USceneComponent> TextRoot;
 
 	/** Each character mesh is held in these components */
-	UPROPERTY(VisibleAnywhere, DuplicateTransient, TextExportTransient, Category="Renderer", meta=(EditCondition="false", EditConditionHides))
+	UPROPERTY(VisibleAnywhere, Category="Renderer", meta=(EditCondition="false", EditConditionHides))
 	TArray<TObjectPtr<UStaticMeshComponent>> CharacterMeshes;
 
 	/** Reusable components to avoid creating and destroying them, recycle components or create new ones when pool is empty */
diff --git a/Engine/Plugins/Experimental/Text3D/Source/Text3D/Private/Utilities/Text3DUtilities.cpp b/Engine/Plugins/Experimental/Text3D/Source/Text3D/Private/Utilities/Text3DUtilities.cpp
--- a/Engine/Plugins/Experimental/Text3D/Source/Text3D/Private/Utilities/Text3DUtilities.cpp
+++ b/Engine/Plugins/Experimental/Text3D/Source/Text3D/Private/Utilities/Text3DUtilities.cpp
@@ -5,10 +5,25 @@
 #include "Engine/Font.h"
 #include "Engine/FontFace.h"
 #include "Fonts/CompositeFont.h"
-#include "Logs/Text3DLogs.h"
 #include "Misc/Paths.h"
+#include "Text3DComponent.h"
 #include "Text3DModule.h"
 
+#if WITH_EDITOR
+#include "Components/StaticMeshComponent.h"
+#include "Dialogs/DlgPickPath.h"
+#include "Engine/StaticMeshActor.h"
+#include "Engine/StaticMesh.h"
+#include "Engine/World.h"
+#include "IMeshMergeUtilities.h"
+#include "Materials/Material.h"
+#include "MeshMergeModule.h"
+#include "MeshMerge/MeshMergingSettings.h"
+#include "Modules/ModuleManager.h"
+#include "Settings/Text3DProjectSettings.h"
+#include "UObject/Package.h"
+#endif
+
 #if WITH_FREETYPE
 THIRD_PARTY_INCLUDES_START
 #include "ft2build.h"
@@ -16,7 +31,9 @@ THIRD_PARTY_INCLUDES_START
 THIRD_PARTY_INCLUDES_END
 #endif
 
-bool UE::Text3D::Utilities::GetSanitizeFontName(const UFont* InFont, FString& OutFontName)
+#define LOCTEXT_NAMESPACE "Text3DUtilities"
+
+bool UE::Text3D::Utilities::Font::GetSanitizeFontName(const UFont* InFont, FString& OutFontName)
 {
 	if (GetFontName(InFont, OutFontName))
 	{
@@ -27,7 +44,7 @@ bool UE::Text3D::Utilities::GetSanitizeFontName(const UFont* InFont, FString& Ou
 	return false;
 }
 
-bool UE::Text3D::Utilities::GetFontName(const UFont* InFont, FString& OutFontName)
+bool UE::Text3D::Utilities::Font::GetFontName(const UFont* InFont, FString& OutFontName)
 {
 	if (!IsValid(InFont))
 	{
@@ -60,7 +77,7 @@ bool UE::Text3D::Utilities::GetFontName(const UFont* InFont, FString& OutFontNam
 	return !OutFontName.IsEmpty();
 }
 
-void UE::Text3D::Utilities::SanitizeFontName(FString& InOutFontName)
+void UE::Text3D::Utilities::Font::SanitizeFontName(FString& InOutFontName)
 {
 	// spaces would be also handled by code below, but spaces used to be removed
 	// so let's to this anyway in order to avoid any potential issues with new vs. old imported fonts
@@ -82,7 +99,7 @@ void UE::Text3D::Utilities::SanitizeFontName(FString& InOutFontName)
 	InOutFontName.RemoveSpacesInline();
 }
 
-bool UE::Text3D::Utilities::GetFontStyle(const UFont* InFont, EText3DFontStyleFlags& OutFontStyleFlags)
+bool UE::Text3D::Utilities::Font::GetFontStyle(const UFont* InFont, EText3DFontStyleFlags& OutFontStyleFlags)
 {
 #if WITH_FREETYPE
 	if (!IsValid(InFont))
@@ -156,7 +173,7 @@ bool UE::Text3D::Utilities::GetFontStyle(const UFont* InFont, EText3DFontStyleFl
 #endif
 }
 
-bool UE::Text3D::Utilities::GetFontStyle(const FText3DFontFamily& InFontFamily, EText3DFontStyleFlags& OutFontStyleFlags)
+bool UE::Text3D::Utilities::Font::GetFontStyle(const FText3DFontFamily& InFontFamily, EText3DFontStyleFlags& OutFontStyleFlags)
 {
 	OutFontStyleFlags = EText3DFontStyleFlags::None;
 
@@ -202,7 +219,7 @@ bool UE::Text3D::Utilities::GetFontStyle(const FText3DFontFamily& InFontFamily,
 #endif
 }
 
-bool UE::Text3D::Utilities::GetFontFaces(const UFont* InFont, TArray<UFontFace*>& OutFontFaces)
+bool UE::Text3D::Utilities::Font::GetFontFaces(const UFont* InFont, TArray<UFontFace*>& OutFontFaces)
 {
 	OutFontFaces.Empty();
 	
@@ -221,3 +238,110 @@ bool UE::Text3D::Utilities::GetFontFaces(const UFont* InFont, TArray<UFontFace*>
 
 	return !OutFontFaces.IsEmpty();
 }
+
+#if WITH_EDITOR
+bool UE::Text3D::Utilities::Conversion::PickAssetPath(const FString& InDefaultPath, FString& OutPickedPath)
+{
+	TSharedPtr<SDlgPickPath> DialogWidget = SNew(SDlgPickPath)
+		.Title(LOCTEXT("PickAssetsLocation", "Choose Asset(s) Location"))
+		.DefaultPath(FText::FromString(InDefaultPath));
+
+	if (DialogWidget->ShowModal() != EAppReturnType::Ok)
+	{
+		OutPickedPath = TEXT("");
+		return false;
+	}
+
+	OutPickedPath = DialogWidget->GetPath().ToString() + TEXT("/");
+	return true;
+}
+
+AStaticMeshActor* UE::Text3D::Utilities::Conversion::ConvertToStaticMesh(const UText3DComponent* InComponent)
+{
+	AStaticMeshActor* NewActor = nullptr;
+
+	if (!IsValid(InComponent))
+	{
+		return NewActor;
+	}
+
+	UWorld* World = InComponent->GetWorld();
+	AActor* Owner = InComponent->GetOwner();
+
+	if (!IsValid(World)
+		|| !IsValid(Owner)
+		|| Owner->bIsEditorPreviewActor
+		)
+	{
+		return NewActor;
+	}
+
+	const UPackage* CurrentPackage = InComponent->GetPackage();
+	FString PackagePath;
+	if (!CurrentPackage || !PickAssetPath(CurrentPackage->GetLoadedPath().GetPackageName(), PackagePath))
+	{
+		return NewActor;
+	}
+
+	const FString AssetPath = PackagePath + Owner->GetActorNameOrLabel() + TEXT("_Merged_") + FString::FromInt(Owner->GetUniqueID());
+
+	TArray<UPrimitiveComponent*> PrimitiveComponents;
+	Owner->GetComponents(PrimitiveComponents);
+
+	if (PrimitiveComponents.IsEmpty())
+	{
+		return NewActor;
+	}
+
+	const IMeshMergeUtilities& MeshMergeUtilities = FModuleManager::Get().LoadModuleChecked<IMeshMergeModule>("MeshMergeUtilities").GetUtilities();
+	FMeshMergingSettings MergeSettings;
+	MergeSettings.bBakeVertexDataToMesh = false;
+	MergeSettings.bMergeMaterials = false;
+	MergeSettings.bMergeEquivalentMaterials = false;
+	MergeSettings.bUseVertexDataForBakingMaterial = false;
+	TArray<UObject*> OutObjects;
+	FVector OutLocation;
+	MeshMergeUtilities.MergeComponentsToStaticMesh(PrimitiveComponents, nullptr, MergeSettings, nullptr, nullptr, AssetPath, OutObjects, OutLocation, 1.f, false);
+
+	if (OutObjects.IsEmpty())
+	{
+		return NewActor;
+	}
+
+	UStaticMesh* Mesh = Cast<UStaticMesh>(OutObjects[0]);
+
+	if (!Mesh)
+	{
+		return NewActor;
+	}
+
+	const UText3DProjectSettings* Text3DSettings = UText3DProjectSettings::Get();
+	for (int32 Index = 0; Index < Mesh->GetStaticMaterials().Num(); Index++)
+	{
+		Mesh->SetMaterial(Index, Text3DSettings->GetDefaultMaterial());
+	}
+
+	// Spawn attached actor with same flags as this actor
+	FActorSpawnParameters SpawnParameters;
+	SpawnParameters.Owner = Owner;
+	SpawnParameters.ObjectFlags = Owner->GetFlags();
+	SpawnParameters.bTemporaryEditorActor = false;
+
+	const FTransform& Transform = InComponent->GetComponentTransform();
+	NewActor = World->SpawnActor<AStaticMeshActor>(AStaticMeshActor::StaticClass(), Transform, SpawnParameters);
+
+	if (!NewActor)
+	{
+		return NewActor;
+	}
+
+	NewActor->SetMobility(EComponentMobility::Movable);
+
+	UStaticMeshComponent* StaticMeshComponent = NewActor->GetStaticMeshComponent();
+	StaticMeshComponent->SetStaticMesh(Mesh);
+
+	return NewActor;
+}
+#endif
+
+#undef LOCTEXT_NAMESPACE
\ No newline at end of file
diff --git a/Engine/Plugins/Experimental/Text3D/Source/Text3D/Public/Utilities/Text3DUtilities.h b/Engine/Plugins/Experimental/Text3D/Source/Text3D/Public/Utilities/Text3DUtilities.h
--- a/Engine/Plugins/Experimental/Text3D/Source/Text3D/Public/Utilities/Text3DUtilities.h
+++ b/Engine/Plugins/Experimental/Text3D/Source/Text3D/Public/Utilities/Text3DUtilities.h
@@ -5,29 +5,45 @@
 #include "HAL/Platform.h"
 #include "Text3DTypes.h"
 
+class AStaticMeshActor;
 class UFontFace;
 class UFont;
+class UText3DComponent;
 struct FText3DTypeFaceMetrics;
 struct FTypefaceEntry;
 
 /** Utilities to manipulate text font */
 namespace UE::Text3D::Utilities
 {
-	/** Get proper font name without invalid chars */
-	TEXT3D_API bool GetSanitizeFontName(const UFont* InFont, FString& OutFontName);
+	namespace Font
+	{
+		/** Get proper font name without invalid chars */
+		TEXT3D_API bool GetSanitizeFontName(const UFont* InFont, FString& OutFontName);
 
-	/** Get font name based on object/import/legacy name */
-	TEXT3D_API bool GetFontName(const UFont* InFont, FString& OutFontName);
+		/** Get font name based on object/import/legacy name */
+		TEXT3D_API bool GetFontName(const UFont* InFont, FString& OutFontName);
 
-	/** Replace invalid object name chars and spaces */
-	TEXT3D_API void SanitizeFontName(FString& InOutFontName);
+		/** Replace invalid object name chars and spaces */
+		TEXT3D_API void SanitizeFontName(FString& InOutFontName);
 
-	/** Detect font style (italic, bold, monospaced) */
-	TEXT3D_API bool GetFontStyle(const UFont* InFont, EText3DFontStyleFlags& OutFontStyleFlags);
+		/** Detect font style (italic, bold, monospaced) */
+		TEXT3D_API bool GetFontStyle(const UFont* InFont, EText3DFontStyleFlags& OutFontStyleFlags);
 	
-	/** Detect font style (italic, bold, monospaced) */
-	TEXT3D_API bool GetFontStyle(const FText3DFontFamily& InFontFamily, EText3DFontStyleFlags& OutFontStyleFlags);
+		/** Detect font style (italic, bold, monospaced) */
+		TEXT3D_API bool GetFontStyle(const FText3DFontFamily& InFontFamily, EText3DFontStyleFlags& OutFontStyleFlags);
 	
-	/** Get font faces that compose the font */
-	TEXT3D_API bool GetFontFaces(const UFont* InFont, TArray<UFontFace*>& OutFontFaces);
+		/** Get font faces that compose the font */
+		TEXT3D_API bool GetFontFaces(const UFont* InFont, TArray<UFontFace*>& OutFontFaces);
+	}
+
+#if WITH_EDITOR
+	namespace Conversion
+	{
+		/** Pick in which folder the asset object will be saved */
+		bool PickAssetPath(const FString& InDefaultPath, FString& OutPickedPath);
+		
+		/** Converts the text component in a static mesh object and spawns a static mesh actor */
+		AStaticMeshActor* ConvertToStaticMesh(const UText3DComponent* InComponent);
+	}
+#endif
 }
\ No newline at end of file
diff --git a/Engine/Plugins/Experimental/Text3D/Source/Text3D/Text3D.Build.cs b/Engine/Plugins/Experimental/Text3D/Source/Text3D/Text3D.Build.cs
--- a/Engine/Plugins/Experimental/Text3D/Source/Text3D/Text3D.Build.cs
+++ b/Engine/Plugins/Experimental/Text3D/Source/Text3D/Text3D.Build.cs
@@ -39,6 +39,12 @@ public class Text3D : ModuleRules
 			PrivateDependencyModuleNames.Add("HarfBuzz");
 		}
 
+		if (Target.Type == TargetType.Editor)
+		{
+			PrivateDependencyModuleNames.Add("MeshMergeUtilities");
+			PrivateDependencyModuleNames.Add("UnrealEd");
+		}
+
 		// Needed to reference underlying FreeType info
 		// @todo: Move font vector handling to new module ("FontCore") - don't expose FT directly?
 		string EnginePath = Path.GetFullPath(Target.RelativeEnginePath);
diff --git a/Engine/Plugins/Experimental/Text3D/Source/Text3DEditor/Private/Customizations/Text3DEditorTextComponentDetailCustomization.cpp b/Engine/Plugins/Experimental/Text3D/Source/Text3DEditor/Private/Customizations/Text3DEditorTextComponentDetailCustomization.cpp
--- a/Engine/Plugins/Experimental/Text3D/Source/Text3DEditor/Private/Customizations/Text3DEditorTextComponentDetailCustomization.cpp
+++ b/Engine/Plugins/Experimental/Text3D/Source/Text3DEditor/Private/Customizations/Text3DEditorTextComponentDetailCustomization.cpp
@@ -2,22 +2,28 @@
 
 #include "Customizations/Text3DEditorTextComponentDetailCustomization.h"
 
+#include "DetailCategoryBuilder.h"
+#include "DetailLayoutBuilder.h"
+#include "DetailWidgetRow.h"
 #include "Modules/ModuleManager.h"
 #include "PropertyEditorModule.h"
+#include "Renderers/Text3DRendererBase.h"
 #include "Text3DComponent.h"
+#include "UObject/Class.h"
+#include "Widgets/Input/SButton.h"
 
 namespace UE::Text3DEditor::Customization
 {
 	void FText3DEditorTextComponentDetailCustomization::CustomizeDetails(IDetailLayoutBuilder& InDetailBuilder)
 	{
+		FPropertyEditorModule& PropertyModule = FModuleManager::GetModuleChecked<FPropertyEditorModule>(TEXT("PropertyEditor"));
+		const FName ComponentClassName = UText3DComponent::StaticClass()->GetFName();
+
 		static bool bSectionInitialized = false;
 		if (!bSectionInitialized)
 		{
 			bSectionInitialized = true;
 
-			FPropertyEditorModule& PropertyModule = FModuleManager::GetModuleChecked<FPropertyEditorModule>(TEXT("PropertyEditor"));
-			const FName ComponentClassName = UText3DComponent::StaticClass()->GetFName();
-
 			const FName GeometrySectionName = TEXT("Geometry");
 			const TSharedRef<FPropertySection> GeometrySection = PropertyModule.FindOrCreateSection(ComponentClassName, GeometrySectionName, FText::FromName(GeometrySectionName));
 			GeometrySection->AddCategory(TEXT("Geometry"));
@@ -43,6 +49,129 @@ namespace UE::Text3DEditor::Customization
 			const FName EffectSectionName = TEXT("Effects");
 			const TSharedRef<FPropertySection> EffectsSection = PropertyModule.FindOrCreateSection(ComponentClassName, EffectSectionName, FText::FromName(EffectSectionName));
 			EffectsSection->AddCategory(TEXT("LayoutEffects"));
+
+			const FName UtilitiesSectionName = TEXT("Utilities");
+			const TSharedRef<FPropertySection> UtilitiesSection = PropertyModule.FindOrCreateSection(ComponentClassName, UtilitiesSectionName, FText::FromName(UtilitiesSectionName));
 		}
+
+		TArray<TWeakObjectPtr<UText3DComponent>> TextComponentsWeak = InDetailBuilder.GetObjectsOfTypeBeingCustomized<UText3DComponent>();
+
+		// Handle ufunctions
+		TMap<FName, FName> FunctionToCategory;
+
+		auto AddUFunction = [this, &FunctionToCategory](UObject* InContext, UFunction* InFunction)
+		{
+			// Only CallInEditor function with 0 parameters
+			if (InFunction
+				&& !InFunction->HasMetaData(TEXT("DeprecatedFunction"))
+				&& InFunction->HasMetaData("CallInEditor")
+				&& InFunction->NumParms == 0)
+			{
+				const FName FunctionName = InFunction->GetFName();
+				TArray<FObjectFunction>& ObjectFunctions = NamedObjectFunctions.FindOrAdd(FunctionName);
+				ObjectFunctions.Add(
+					{
+						.Owner = InContext,
+						.Function = InFunction
+					}
+				);
+				FunctionToCategory.Add(FunctionName, FName(InFunction->GetMetaData(TEXT("Category"))));
+			}
+		};
+
+		for (const TWeakObjectPtr<UText3DComponent>& TextComponentWeak : TextComponentsWeak)
+		{
+			UText3DComponent* TextComponent = TextComponentWeak.Get();
+
+			if (!TextComponent)
+			{
+				continue;
+			}
+
+			// Look for ufunction in component
+			for (UFunction* Function : TFieldRange<UFunction>(TextComponent->GetClass(), EFieldIteratorFlags::ExcludeSuper))
+			{
+				AddUFunction(TextComponent, Function);
+			}
+
+			// Look for ufunction in renderer
+			if (UText3DRendererBase* Renderer = TextComponent->GetTextRenderer())
+			{
+				for (UFunction* Function : TFieldRange<UFunction>(Renderer->GetClass(), EFieldIteratorFlags::ExcludeSuper))
+				{
+					AddUFunction(Renderer, Function);
+				}
+			}
+		}
+
+		if (!FunctionToCategory.IsEmpty())
+		{
+			TMap<FName, const TSharedPtr<SVerticalBox>> FunctionToWidget;
+
+			// Add buttons for ufunctions based on their category
+			for (const TPair<FName, FName>& FunctionToCategoryPair : FunctionToCategory)
+			{
+				const FName FunctionsCategoryName = FunctionToCategoryPair.Value;
+
+				IDetailCategoryBuilder& FunctionsCategory = InDetailBuilder.EditCategory(FunctionsCategoryName, FText::FromName(FunctionsCategoryName), ECategoryPriority::Uncommon);
+
+				// Subcategories are not supported in sections
+				if (!FunctionsCategoryName.ToString().Contains(TEXT("|")))
+				{
+					const TSharedRef<FPropertySection> FunctionsSection = PropertyModule.FindOrCreateSection(ComponentClassName, FunctionsCategoryName, FText::FromName(FunctionsCategoryName));
+					FunctionsSection->AddCategory(FunctionsCategoryName);
+				}
+
+				const TSharedPtr<SVerticalBox> FunctionsWidget = SNew(SVerticalBox);
+
+				FunctionsCategory.AddCustomRow(FText::GetEmpty())
+					.WholeRowContent()
+					.HAlign(HAlign_Left)
+					[
+						FunctionsWidget.ToSharedRef()
+					];
+
+				FunctionToWidget.Add(FunctionToCategoryPair.Key, FunctionsWidget);
+			}
+
+			for (const TPair<FName, const TSharedPtr<SVerticalBox>>& FunctionNameToWidget : FunctionToWidget)
+			{
+				const FText ButtonLabel = FText::FromString(FName::NameToDisplayString(FunctionNameToWidget.Key.ToString(), /** IsBool */false));
+
+				const TSharedPtr<SVerticalBox>& FunctionsWidget = FunctionNameToWidget.Value;
+
+				FunctionsWidget->AddSlot()
+					.Padding(0.f, 3.f)
+					.AutoHeight()
+					[
+						SNew(SButton)
+						.Text(ButtonLabel)
+						.HAlign(HAlign_Center)
+						.VAlign(VAlign_Fill)
+						.OnClicked(this, &FText3DEditorTextComponentDetailCustomization::OnFunctionButtonClicked, FunctionNameToWidget.Key)
+					];
+			}
+		}
+	}
+
+	FReply FText3DEditorTextComponentDetailCustomization::OnFunctionButtonClicked(FName InFunctionName)
+	{
+		if (TArray<FObjectFunction> const* ObjectFunctions = NamedObjectFunctions.Find(InFunctionName))
+		{
+			for (const FObjectFunction& ObjectFunction : *ObjectFunctions)
+			{
+				UObject* Object = ObjectFunction.Owner.Get();
+				UFunction* Function = ObjectFunction.Function.Get();
+
+				if (!Object || !Function)
+				{
+					continue;
+				}
+
+				Object->ProcessEvent(Function, nullptr);
+			}
+		}
+
+		return FReply::Handled();
 	}
 }
diff --git a/Engine/Plugins/Experimental/Text3D/Source/Text3DEditor/Private/Customizations/Text3DEditorTextComponentDetailCustomization.h b/Engine/Plugins/Experimental/Text3D/Source/Text3DEditor/Private/Customizations/Text3DEditorTextComponentDetailCustomization.h
--- a/Engine/Plugins/Experimental/Text3D/Source/Text3DEditor/Private/Customizations/Text3DEditorTextComponentDetailCustomization.h
+++ b/Engine/Plugins/Experimental/Text3D/Source/Text3DEditor/Private/Customizations/Text3DEditorTextComponentDetailCustomization.h
@@ -2,11 +2,16 @@
 
 #pragma once
 
+#include "Containers/Map.h"
 #include "IDetailCustomization.h"
 #include "Templates/SharedPointerFwd.h"
+#include "UObject/WeakObjectPtr.h"
 
+class FReply;
 class IPropertyHandle;
 class IPropertyUtilities;
+class UFunction;
+class UObject;
 
 namespace UE::Text3DEditor::Customization
 {
@@ -22,5 +27,18 @@ namespace UE::Text3DEditor::Customization
 		//~ Begin IDetailCustomization
 		virtual void CustomizeDetails(IDetailLayoutBuilder& InDetailBuilder) override;
 		//~ End IDetailCustomization
+
+	protected:
+		struct FObjectFunction
+		{
+			TWeakObjectPtr<UObject> Owner;
+			TWeakObjectPtr<UFunction> Function;
+		};
+
+		/** Execute ufunction with that name on selected objects */
+		FReply OnFunctionButtonClicked(FName InFunctionName);
+
+		/** Function name to object ufunction mapping */
+		TMap<FName, TArray<FObjectFunction>> NamedObjectFunctions;
 	};
 }
\ No newline at end of file
diff --git a/Engine/Plugins/Experimental/Text3D/Source/Text3DEditor/Private/Subsystems/Text3DEditorFontSubsystem.cpp b/Engine/Plugins/Experimental/Text3D/Source/Text3DEditor/Private/Subsystems/Text3DEditorFontSubsystem.cpp
--- a/Engine/Plugins/Experimental/Text3D/Source/Text3DEditor/Private/Subsystems/Text3DEditorFontSubsystem.cpp
+++ b/Engine/Plugins/Experimental/Text3D/Source/Text3DEditor/Private/Subsystems/Text3DEditorFontSubsystem.cpp
@@ -159,14 +159,14 @@ const FText3DEditorFont* UText3DEditorFontSubsystem::GetEditorFont(const FString
 const FText3DEditorFont* UText3DEditorFontSubsystem::GetSystemFont(const FString& InFontName) const
 {
 	FString SanitizedFontName = InFontName;
-	UE::Text3D::Utilities::SanitizeFontName(SanitizedFontName);
+	UE::Text3D::Utilities::Font::SanitizeFontName(SanitizedFontName);
 	return SystemFonts.Find(SanitizedFontName);
 }
 
 const FText3DEditorFont* UText3DEditorFontSubsystem::GetProjectFont(const FString& InFontName) const
 {
 	FString SanitizedFontName = InFontName;
-	UE::Text3D::Utilities::SanitizeFontName(SanitizedFontName);
+	UE::Text3D::Utilities::Font::SanitizeFontName(SanitizedFontName);
 	return ProjectFonts.Find(SanitizedFontName);
 }
 
@@ -178,7 +178,7 @@ const FText3DEditorFont* UText3DEditorFontSubsystem::FindEditorFont(const UFont*
 	}
 
 	FString SanitizedFontName;
-	if (!UE::Text3D::Utilities::GetSanitizeFontName(InFont, SanitizedFontName))
+	if (!UE::Text3D::Utilities::Font::GetSanitizeFontName(InFont, SanitizedFontName))
 	{
 		UE_LOG(LogText3DEditor, Warning, TEXT("Could not retrieve sanitized font name %s"), *InFont->GetName());
 		return nullptr;
@@ -365,7 +365,7 @@ bool UText3DEditorFontSubsystem::UnregisterProjectFont(const UFont* InFont)
 	}
 
 	FString FontName;
-	if (!UE::Text3D::Utilities::GetSanitizeFontName(InFont, FontName))
+	if (!UE::Text3D::Utilities::Font::GetSanitizeFontName(InFont, FontName))
 	{
 		return false;
 	}
@@ -391,7 +391,7 @@ bool UText3DEditorFontSubsystem::RegisterSystemFont(const FText3DFontFamily& InF
 	}
 
 	FString SanitizedFontFamilyName = InFontFamily.FontFamilyName;
-	UE::Text3D::Utilities::SanitizeFontName(SanitizedFontFamilyName);
+	UE::Text3D::Utilities::Font::SanitizeFontName(SanitizedFontFamilyName);
 
 	UFont* NewFont = NewObject<UFont>(this, *SanitizedFontFamilyName, RF_Public | RF_Standalone);
 	NewFont->ImportOptions.FontName = SanitizedFontFamilyName;
@@ -409,7 +409,7 @@ bool UText3DEditorFontSubsystem::RegisterSystemFont(const FText3DFontFamily& InF
 		}
 
 		FString SanitizedFontFaceName = FontFacePath.Key;
-		UE::Text3D::Utilities::SanitizeFontName(SanitizedFontFaceName);
+		UE::Text3D::Utilities::Font::SanitizeFontName(SanitizedFontFaceName);
 		const FString FontFaceAssetName = SanitizedFontFamilyName + "_" + SanitizedFontFaceName;
 
 		bool bCanceled = false;
@@ -432,10 +432,10 @@ bool UText3DEditorFontSubsystem::RegisterSystemFont(const FText3DFontFamily& InF
 	}
 	
 	TArray<UFontFace*> FontFaces;
-	UE::Text3D::Utilities::GetFontFaces(NewFont, FontFaces);
+	UE::Text3D::Utilities::Font::GetFontFaces(NewFont, FontFaces);
 
 	EText3DFontStyleFlags FontStyleFlags;
-	UE::Text3D::Utilities::GetFontStyle(NewFont, FontStyleFlags);
+	UE::Text3D::Utilities::Font::GetFontStyle(NewFont, FontStyleFlags);
 
 	FText3DEditorFont& RegisteredFont = SystemFonts.FindOrAdd(SanitizedFontFamilyName);
 	
@@ -515,21 +515,21 @@ bool UText3DEditorFontSubsystem::RegisterProjectFont(UFont* InFont)
 	}
 
 	FString SanitizedFontName;
-	if (!UE::Text3D::Utilities::GetSanitizeFontName(InFont, SanitizedFontName))
+	if (!UE::Text3D::Utilities::Font::GetSanitizeFontName(InFont, SanitizedFontName))
 	{
 		UE_LOG(LogText3DEditor, Warning, TEXT("Could not retrieve sanitized font name %s"), *InFont->GetName());
 		return false;
 	}
 
 	EText3DFontStyleFlags FontStyleFlags;
-	if (!UE::Text3D::Utilities::GetFontStyle(InFont, FontStyleFlags))
+	if (!UE::Text3D::Utilities::Font::GetFontStyle(InFont, FontStyleFlags))
 	{
 		UE_LOG(LogText3DEditor, Warning, TEXT("Could not retrieve font style for %s"), *InFont->GetName());
 		return false;
 	}
 
 	TArray<UFontFace*> FontFaces;
-	if (!UE::Text3D::Utilities::GetFontFaces(InFont, FontFaces))
+	if (!UE::Text3D::Utilities::Font::GetFontFaces(InFont, FontFaces))
 	{
 		UE_LOG(LogText3DEditor, Warning, TEXT("Could not retrieve font faces for %s"), *InFont->GetName());
 		return false;
