MotionDesign : Text3D
- Fixed older level with Text not appearing in rundown/editor due to how component is initialized
- Initializing text when renderer is changed to avoid other system updates being processed before text is generated and causing issues/crashes
- Fixed text character index selector reset to default

diff --git a/Engine/Plugins/Experimental/Text3D/Source/Text3D/Private/Extensions/Text3DDefaultGeometryExtension.cpp b/Engine/Plugins/Experimental/Text3D/Source/Text3D/Private/Extensions/Text3DDefaultGeometryExtension.cpp
--- a/Engine/Plugins/Experimental/Text3D/Source/Text3D/Private/Extensions/Text3DDefaultGeometryExtension.cpp
+++ b/Engine/Plugins/Experimental/Text3D/Source/Text3D/Private/Extensions/Text3DDefaultGeometryExtension.cpp
@@ -161,12 +161,12 @@ EText3DVerticalTextAlignment UText3DDefaultGeometryExtension::GetGlyphVAlignment
 
 FTypefaceFontData* UText3DDefaultGeometryExtension::GetTypefaceFontData() const
 {
+	FTypefaceFontData* TypefaceFontData = nullptr;
 	if (CachedFontData)
 	{
-		return CachedFontData->GetTypefaceFontData(GlyphMeshParameters.TypefaceIndex);
+		TypefaceFontData = CachedFontData->GetTypefaceFontData(GlyphMeshParameters.TypefaceIndex);
 	}
-
-	return nullptr;
+	return TypefaceFontData;
 }
 
 EText3DExtensionResult UText3DDefaultGeometryExtension::PreRendererUpdate(const UE::Text3D::Renderer::FUpdateParameters& InParameters)
@@ -177,10 +177,16 @@ EText3DExtensionResult UText3DDefaultGeometryExtension::PreRendererUpdate(const
 	}
 
 	const UText3DComponent* TextComponent = GetText3DComponent();
-	UText3DEngineSubsystem* Subsystem = UText3DEngineSubsystem::Get();
 	const uint32 TypefaceIndex = TextComponent->GetTypeFaceIndex();
 	UFont* Font = TextComponent->GetFont();
 
+	if (TypefaceIndex == INDEX_NONE)
+	{
+		UE_LOG(LogText3D, Error, TEXT("Invalid typeface index returned for font %s : %i"), Font ? *Font->GetName() : TEXT("invalid"), TypefaceIndex);
+		return EText3DExtensionResult::Failed;
+	}
+
+	UText3DEngineSubsystem* Subsystem = UText3DEngineSubsystem::Get();
 	CachedFontData = &Subsystem->GetCachedFontData(Font, TypefaceIndex);
 
 	if (!CachedFontData)
diff --git a/Engine/Plugins/Experimental/Text3D/Source/Text3D/Private/Extensions/Text3DDefaultLayoutExtension.cpp b/Engine/Plugins/Experimental/Text3D/Source/Text3D/Private/Extensions/Text3DDefaultLayoutExtension.cpp
--- a/Engine/Plugins/Experimental/Text3D/Source/Text3D/Private/Extensions/Text3DDefaultLayoutExtension.cpp
+++ b/Engine/Plugins/Experimental/Text3D/Source/Text3D/Private/Extensions/Text3DDefaultLayoutExtension.cpp
@@ -282,6 +282,7 @@ EText3DExtensionResult UText3DDefaultLayoutExtension::PreRendererUpdate(const UE
 		FTypefaceFontData* TypefaceFontData = GeometryExtension->GetTypefaceFontData();
 		if (!TypefaceFontData)
 		{
+			UE_LOG(LogText3D, Warning, TEXT("Typeface font data cannot be retrieved, cannot proceed without it"))
 			return EText3DExtensionResult::Failed;
 		}
 
diff --git a/Engine/Plugins/Experimental/Text3D/Source/Text3D/Private/Renderers/StaticMeshes/Text3DStaticMeshesRenderer.cpp b/Engine/Plugins/Experimental/Text3D/Source/Text3D/Private/Renderers/StaticMeshes/Text3DStaticMeshesRenderer.cpp
--- a/Engine/Plugins/Experimental/Text3D/Source/Text3D/Private/Renderers/StaticMeshes/Text3DStaticMeshesRenderer.cpp
+++ b/Engine/Plugins/Experimental/Text3D/Source/Text3D/Private/Renderers/StaticMeshes/Text3DStaticMeshesRenderer.cpp
@@ -210,7 +210,6 @@ void UText3DStaticMeshesRenderer::AllocateComponents(int32 InCount)
 			{
 				MeshComponent->DetachFromComponent(FDetachmentTransformRules::KeepRelativeTransform);
 				MeshComponent->UnregisterComponent();
-				MeshComponent->SetStaticMesh(nullptr);
 				Actor->RemoveInstanceComponent(MeshComponent);
 				Actor->RemoveOwnedComponent(MeshComponent);
 				CharacterMeshesPool.Add(MeshComponent);
@@ -230,10 +229,12 @@ void UText3DStaticMeshesRenderer::AllocateComponents(int32 InCount)
 			{
 				const FName StaticMeshComponentName = MakeUniqueObjectName(this, UStaticMeshComponent::StaticClass(), FName(TEXT("CharacterMesh")));
 				MeshComponent = NewObject<UStaticMeshComponent>(this, StaticMeshComponentName);
+				MeshComponent->OnComponentCreated();
 				CharacterMeshes.Add(MeshComponent);
 			}
 
-			MeshComponent->SetupAttachment(TextRoot);
+			MeshComponent->RegisterComponent();
+			MeshComponent->AttachToComponent(TextRoot, FAttachmentTransformRules::KeepRelativeTransform);
 			Actor->AddInstanceComponent(MeshComponent);
 			Actor->AddOwnedComponent(MeshComponent);
 		}
diff --git a/Engine/Plugins/Experimental/Text3D/Source/Text3D/Private/Text3DComponent.cpp b/Engine/Plugins/Experimental/Text3D/Source/Text3D/Private/Text3DComponent.cpp
--- a/Engine/Plugins/Experimental/Text3D/Source/Text3D/Private/Text3DComponent.cpp
+++ b/Engine/Plugins/Experimental/Text3D/Source/Text3D/Private/Text3DComponent.cpp
@@ -85,6 +85,7 @@ void UText3DComponent::PostLoad()
 {
 	Super::PostLoad();
 
+	RefreshTypeface();
 	OnTextRendererClassChanged();
 }
 
@@ -92,6 +93,11 @@ void UText3DComponent::RequestUpdate(EText3DRendererFlags InFlags, bool bInImmed
 {
 	if (InFlags != EText3DRendererFlags::None)
 	{
+		if (IsTemplate(RF_ClassDefaultObject))
+		{
+			return;
+		}
+
 		if (bInImmediate)
 		{
 			Rebuild(InFlags, /*bCleanCache*/false);
@@ -103,15 +109,15 @@ void UText3DComponent::RequestUpdate(EText3DRendererFlags InFlags, bool bInImmed
 	}
 }
 
-uint32 UText3DComponent::GetTypeFaceIndex() const
+int32 UText3DComponent::GetTypeFaceIndex() const
 {
-	uint32 TypefaceIndex = INDEX_NONE;
+	int32 TypefaceIndex = INDEX_NONE;
 
 	if (Font)
 	{
 		TArray<FTypefaceEntry>& Fonts = Font->CompositeFont.DefaultTypeface.Fonts;
 
-		for (uint32 Index = 0; Index < static_cast<uint32>(Fonts.Num()); Index++)
+		for (int32 Index = 0; Index < Fonts.Num(); Index++)
 		{
 			if (Typeface == Fonts[Index].Name)
 			{
@@ -130,7 +136,7 @@ const FTypefaceEntry* UText3DComponent::GetTypeFaceEntry() const
 	{
 		TArray<FTypefaceEntry>& TypefaceEntries = Font->CompositeFont.DefaultTypeface.Fonts;
 
-		for (uint32 Index = 0; Index < static_cast<uint32>(TypefaceEntries.Num()); Index++)
+		for (int32 Index = 0; Index < TypefaceEntries.Num(); Index++)
 		{
 			if (Typeface == TypefaceEntries[Index].Name)
 			{
@@ -251,10 +257,10 @@ void UText3DComponent::OnRegister()
 	if (TextRenderer)
 	{
 		TextRenderer->Create();
+		const EText3DRendererFlags Flags = bInitialized ? EText3DRendererFlags::None : EText3DRendererFlags::All;
+		constexpr bool bImmediate = true;
+		RequestUpdate(Flags, bImmediate);
 	}
-
-	constexpr bool bImmediate = true;
-	RequestUpdate(UpdateFlags, bImmediate);
 }
 
 void UText3DComponent::OnUnregister()
@@ -930,21 +936,6 @@ void UText3DComponent::Rebuild(EText3DRendererFlags InUpdateFlags, bool bCleanCa
 		return;
 	}
 
-	if (IsTemplate(RF_ClassDefaultObject))
-	{
-		return;
-	}
-
-	if (!IsRegistered())
-	{
-		return;
-	}
-
-	if (UE::IsSavingPackage(this))
-	{
-		return;
-	}
-
 	if (!TextRenderer)
 	{
 		UE_LOG(LogText3D, Error, TEXT("Text3D renderer is not valid, cannot proceed"));
@@ -964,10 +955,17 @@ void UText3DComponent::Rebuild(EText3DRendererFlags InUpdateFlags, bool bCleanCa
 			return;
 		}
 	}
-	
+
+	if (UE::IsSavingPackage(this))
+	{
+		UE_LOG(LogText3D, Warning, TEXT("[%s] Rebuild called with update flags (%u) while package is being saved. This new Rebuild will be deferred to the next tick."), *GetReadableName(), InUpdateFlags);
+		QueueRebuild(InUpdateFlags);
+		return;
+	}
+
 	if (bIsUpdatingText)
 	{
-		UE_LOG(LogText3D, Warning, TEXT("[%s] Rebuild called with update flags (%u) while building text. This new Rebuild will be deferred to the next tick."), *GetReadableName(), static_cast<uint8>(UpdateFlags));
+		UE_LOG(LogText3D, Warning, TEXT("[%s] Rebuild called with update flags (%u) while building text. This new Rebuild will be deferred to the next tick."), *GetReadableName(), InUpdateFlags);
 		QueueRebuild(InUpdateFlags);
 		return;
 	}
@@ -979,6 +977,13 @@ void UText3DComponent::Rebuild(EText3DRendererFlags InUpdateFlags, bool bCleanCa
 		bIsUpdatingText = false;
 	};
 
+	if (!bInitialized)
+	{
+		UE_LOG(LogText3D, Warning, TEXT("[%s] Rebuild called with update flags (%u) while text is not initialized yet."), *GetReadableName(), InUpdateFlags);
+		InUpdateFlags = EText3DRendererFlags::All;
+		bInitialized = true;
+	}
+
 	TextPreUpdateDelegate.Broadcast(this, InUpdateFlags);
 
 	TRACE_CPUPROFILER_EVENT_SCOPE(UText3DComponent::RebuildInternal);
@@ -1055,6 +1060,8 @@ void UText3DComponent::OnTextRendererClassChanged()
 	const AActor* ContextActor = GetOwner();
 	UE_LOG(LogText3D, Log, TEXT("%s : Text3D renderer changed : old %s - new %s"), !!ContextActor ? *ContextActor->GetActorNameOrLabel() : TEXT("Invalid owner"), *OldRendererName.ToString(), *TextRenderer->GetFriendlyName().ToString());
 
+	bInitialized = false;
+
 	if (IsRegistered())
 	{
 		TextRenderer->Create();
diff --git a/Engine/Plugins/Experimental/Text3D/Source/Text3D/Public/Text3DComponent.h b/Engine/Plugins/Experimental/Text3D/Source/Text3D/Public/Text3DComponent.h
--- a/Engine/Plugins/Experimental/Text3D/Source/Text3D/Public/Text3DComponent.h
+++ b/Engine/Plugins/Experimental/Text3D/Source/Text3D/Public/Text3DComponent.h
@@ -262,7 +262,7 @@ public:
 
 	TEXT3D_API void RequestUpdate(EText3DRendererFlags InFlags, bool bInImmediate = false);
 
-	uint32 GetTypeFaceIndex() const;
+	int32 GetTypeFaceIndex() const;
 
 	const FTypefaceEntry* GetTypeFaceEntry() const;
 	
@@ -482,7 +482,7 @@ private:
 	std::atomic<bool> bIsUpdatingText;
 
 	/** Used to determine and selectively perform the the type of rebuild requested. */
-	EText3DRendererFlags UpdateFlags = EText3DRendererFlags::All;
+	EText3DRendererFlags UpdateFlags = EText3DRendererFlags::None;
 
 	/** Text statistics cached since last text generation */
 	FText3DStatistics Statistics;
@@ -506,6 +506,9 @@ private:
 
 	bool bZeroScaled = false;
 
+	/** Was the text initialized once */
+	bool bInitialized = false;
+
 #if WITH_EDITORONLY_DATA
 	UE_DEPRECATED(5.6, "Use Material Extension instead")
 	UPROPERTY()
diff --git a/Engine/Plugins/Experimental/Text3D/Source/Text3DEditor/Private/Customizations/Text3DEditorCharacterPropertyTypeCustomization.cpp b/Engine/Plugins/Experimental/Text3D/Source/Text3DEditor/Private/Customizations/Text3DEditorCharacterPropertyTypeCustomization.cpp
--- a/Engine/Plugins/Experimental/Text3D/Source/Text3DEditor/Private/Customizations/Text3DEditorCharacterPropertyTypeCustomization.cpp
+++ b/Engine/Plugins/Experimental/Text3D/Source/Text3DEditor/Private/Customizations/Text3DEditorCharacterPropertyTypeCustomization.cpp
@@ -41,6 +41,19 @@ namespace UE::Text3DEditor::Customization
 		CharacterExtensionWeak = CharacterExtension;
 		CharacterPropertyHandleWeak = InPropertyHandle;
 
+		InPropertyHandle->SetOnPropertyResetToDefault(FSimpleDelegate::CreateSPLambda(this, [this]()
+		{
+			if (const TSharedPtr<IPropertyHandle> PropertyHandle = CharacterPropertyHandleWeak.Pin())
+			{
+				PropertyHandle->SetValue(0);
+				OnTextCharacterChanged();
+				if (const TSharedPtr<IPropertyUtilities> Utilities = PropertyUtilitiesWeak.Pin())
+				{
+					Utilities->RequestForceRefresh();
+				}
+			}
+		}));
+
 		OnTextCharacterChanged();
 
 		if (UText3DComponent* Text3DComponent = CharacterExtension->GetText3DComponent())
@@ -168,12 +181,9 @@ namespace UE::Text3DEditor::Customization
 	{
 		if (const UText3DCharacterExtensionBase* Extension = CharacterExtensionWeak.Get())
 		{
-			const UText3DCharacterBase* Character = Extension->GetCharacter(Extension->TextCharacterIndex);
-
 			return FText::Format(
-				LOCTEXT("CharacterCount", "Total: {0} / Current: {1}")
+				LOCTEXT("CharacterCount", "Total: {0}")
 				, FText::AsNumber(Extension->GetCharacterCount())
-				, IsValid(Character) ? FText::FromString(Character->GetCharacter()) : FText::GetEmpty()
 			);
 		}
 
