Fixed deprecation warnings

#rnx

diff --git a/Engine/Plugins/Experimental/Text3D/Source/Text3DEditor/Private/Subsystems/Text3DEditorFontSubsystem.cpp b/Engine/Plugins/Experimental/Text3D/Source/Text3DEditor/Private/Subsystems/Text3DEditorFontSubsystem.cpp
--- a/Engine/Plugins/Experimental/Text3D/Source/Text3DEditor/Private/Subsystems/Text3DEditorFontSubsystem.cpp
+++ b/Engine/Plugins/Experimental/Text3D/Source/Text3DEditor/Private/Subsystems/Text3DEditorFontSubsystem.cpp
@@ -415,7 +415,7 @@ bool UText3DEditorFontSubsystem::RegisterSystemFont(const FText3DFontFamily& InF
 		bool bCanceled = false;
 		if (UFontFace* NewFontFace = Cast<UFontFace>(FontFaceFactory->ImportObject(UFontFace::StaticClass(), this, *FontFaceAssetName, RF_Public | RF_Standalone, FontFacePath.Value, TEXT(""), bCanceled)))
 		{
-			FTypefaceEntry& Entry = NewFont->CompositeFont.DefaultTypeface.Fonts.AddDefaulted_GetRef();
+			FTypefaceEntry& Entry = NewFont->GetMutableInternalCompositeFont().DefaultTypeface.Fonts.AddDefaulted_GetRef();
 			Entry.Name = FName(FName::NameToDisplayString(SanitizedFontFaceName, false));
 			Entry.Font = FFontData(NewFontFace);
 		}
@@ -426,7 +426,7 @@ bool UText3DEditorFontSubsystem::RegisterSystemFont(const FText3DFontFamily& InF
 	// Fallback
 	if (const UText3DProjectSettings* Text3DProjectSettings = UText3DProjectSettings::Get())
 	{
-		FTypefaceEntry& FallbackTypefaceEntry = NewFont->CompositeFont.FallbackTypeface.Typeface.Fonts.AddDefaulted_GetRef();
+		FTypefaceEntry& FallbackTypefaceEntry = NewFont->GetMutableInternalCompositeFont().FallbackTypeface.Typeface.Fonts.AddDefaulted_GetRef();
 		FallbackTypefaceEntry.Font = FFontData(Text3DProjectSettings->GetFallbackFontFace());
 		FallbackTypefaceEntry.Name = FName("Regular");
 	}
