Fixed warning caused by 43864510

UText3DEditorFontSubsystem emits a warning if it finds a UFont asset using a FCompositeFont, that isn't using UFontFace assets internally (as is the case for FCoreStyle::GetDefaultFont()). This just has it skip such fonts instead.

#rnx

diff --git a/Engine/Plugins/Experimental/Text3D/Source/Text3DEditor/Private/Subsystems/Text3DEditorFontSubsystem.cpp b/Engine/Plugins/Experimental/Text3D/Source/Text3DEditor/Private/Subsystems/Text3DEditorFontSubsystem.cpp
--- a/Engine/Plugins/Experimental/Text3D/Source/Text3DEditor/Private/Subsystems/Text3DEditorFontSubsystem.cpp
+++ b/Engine/Plugins/Experimental/Text3D/Source/Text3DEditor/Private/Subsystems/Text3DEditorFontSubsystem.cpp
@@ -508,7 +508,7 @@ bool UText3DEditorFontSubsystem::RegisterProjectFont(UFont* InFont)
 		return false;
 	}
 
-	if (!InFont->GetCompositeFont() || InFont->GetCompositeFont()->DefaultTypeface.Fonts.IsEmpty())
+	if (!InFont->GetCompositeFont() || InFont->GetCompositeFont()->DefaultTypeface.Fonts.IsEmpty() || InFont->RuntimeFontSource != ERuntimeFontSource::Asset)
 	{
 		UE_LOG(LogText3DEditor, Log, TEXT("No composite font found for font %s"), *InFont->GetName());
 		return false;
