MotionDesign : Text3D
- Added text debug mode in text settings with custom renderer behavior, static mesh renderer will show/hide components tree hierarchy

diff --git a/Engine/Plugins/Experimental/Text3D/Source/Text3D/Private/Renderers/StaticMeshes/Text3DStaticMeshesRenderer.cpp b/Engine/Plugins/Experimental/Text3D/Source/Text3D/Private/Renderers/StaticMeshes/Text3DStaticMeshesRenderer.cpp
--- a/Engine/Plugins/Experimental/Text3D/Source/Text3D/Private/Renderers/StaticMeshes/Text3DStaticMeshesRenderer.cpp
+++ b/Engine/Plugins/Experimental/Text3D/Source/Text3D/Private/Renderers/StaticMeshes/Text3DStaticMeshesRenderer.cpp
@@ -30,6 +30,11 @@ void UText3DStaticMeshesRenderer::OnCreate()
 		TextRoot = NewObject<USceneComponent>(this, TEXT("TextRoot"));
 		TextRoot->SetupAttachment(TextComponent);
 	}
+
+	for (const TObjectPtr<UStaticMeshComponent>& CharacterMesh : CharacterMeshes)
+	{
+		CharacterMesh->RegisterComponent();
+	}
 }
 
 void UText3DStaticMeshesRenderer::OnUpdate(const UE::Text3D::Renderer::FUpdateParameters& InParameters)
@@ -229,13 +234,11 @@ void UText3DStaticMeshesRenderer::AllocateComponents(int32 InCount)
 			{
 				const FName StaticMeshComponentName = MakeUniqueObjectName(this, UStaticMeshComponent::StaticClass(), FName(TEXT("CharacterMesh")));
 				MeshComponent = NewObject<UStaticMeshComponent>(this, StaticMeshComponentName);
-				MeshComponent->OnComponentCreated();
 				CharacterMeshes.Add(MeshComponent);
 			}
 
 			MeshComponent->RegisterComponent();
 			MeshComponent->AttachToComponent(TextRoot, FAttachmentTransformRules::KeepRelativeTransform);
-			Actor->AddInstanceComponent(MeshComponent);
 			Actor->AddOwnedComponent(MeshComponent);
 		}
 	}
@@ -297,6 +300,46 @@ void UText3DStaticMeshesRenderer::ConvertToStaticMesh()
 		UE_LOG(LogText3D, Warning, TEXT("%s : ConvertToStaticMesh Failed"), *Owner->GetActorNameOrLabel())
 	}
 }
+
+void UText3DStaticMeshesRenderer::OnDebugModeEnabled()
+{
+	Super::OnDebugModeEnabled();
+	SetDebugMode(true);
+}
+
+void UText3DStaticMeshesRenderer::OnDebugModeDisabled()
+{
+	Super::OnDebugModeDisabled();
+	SetDebugMode(false);
+}
+
+void UText3DStaticMeshesRenderer::SetDebugMode(bool bEnabled)
+{
+	// Since we are dealing with class FProperty, no need to run this for each instance, do it once
+	static bool bDebugModeEnabled = true;
+
+	if (bDebugModeEnabled != bEnabled)
+	{
+		bDebugModeEnabled = bEnabled;
+
+		FProperty* TextRootProperty = FindFProperty<FProperty>(UText3DStaticMeshesRenderer::StaticClass(), GET_MEMBER_NAME_CHECKED(UText3DStaticMeshesRenderer, TextRoot));
+		FArrayProperty* CharacterMeshesProperty = FindFProperty<FArrayProperty>(UText3DStaticMeshesRenderer::StaticClass(), GET_MEMBER_NAME_CHECKED(UText3DStaticMeshesRenderer, CharacterMeshes));
+
+		// Here we toggle the CPF_Edit flag to hide/show property in the details panel component editor tree / outliner
+		// @see FComponentEditorUtils::GetPropertyForEditableNativeComponent
+		// todo : implement a custom debug view widget for text or add a property editor metadata to control the component visibility in the component editor tree / outliner
+		if (bEnabled)
+		{
+			TextRootProperty->SetPropertyFlags(CPF_Edit);
+			CharacterMeshesProperty->SetPropertyFlags(CPF_Edit);
+		}
+		else
+		{
+			TextRootProperty->ClearPropertyFlags(CPF_Edit);
+			CharacterMeshesProperty->ClearPropertyFlags(CPF_Edit);
+		}
+	}
+}
 #endif
 
 #undef LOCTEXT_NAMESPACE 
\ No newline at end of file
diff --git a/Engine/Plugins/Experimental/Text3D/Source/Text3D/Private/Renderers/StaticMeshes/Text3DStaticMeshesRenderer.h b/Engine/Plugins/Experimental/Text3D/Source/Text3D/Private/Renderers/StaticMeshes/Text3DStaticMeshesRenderer.h
--- a/Engine/Plugins/Experimental/Text3D/Source/Text3D/Private/Renderers/StaticMeshes/Text3DStaticMeshesRenderer.h
+++ b/Engine/Plugins/Experimental/Text3D/Source/Text3D/Private/Renderers/StaticMeshes/Text3DStaticMeshesRenderer.h
@@ -53,10 +53,18 @@ protected:
 	virtual void OnDestroy() override;
 	virtual FName GetFriendlyName() const override;
 	virtual FBox OnCalculateBounds() const override;
+#if WITH_EDITOR
+	virtual void OnDebugModeEnabled() override;
+	virtual void OnDebugModeDisabled() override;
+#endif
 	//~ End UText3DRendererBase
 
 	void AllocateComponents(int32 InCount);
 
+#if WITH_EDITOR
+	void SetDebugMode(bool bEnabled);
+#endif
+
 	/** Holds all character components */
 	UPROPERTY(VisibleAnywhere, Category="Renderer", meta=(EditCondition="false", EditConditionHides))
 	TObjectPtr<USceneComponent> TextRoot;
diff --git a/Engine/Plugins/Experimental/Text3D/Source/Text3D/Private/Renderers/Text3DRendererBase.cpp b/Engine/Plugins/Experimental/Text3D/Source/Text3D/Private/Renderers/Text3DRendererBase.cpp
--- a/Engine/Plugins/Experimental/Text3D/Source/Text3D/Private/Renderers/Text3DRendererBase.cpp
+++ b/Engine/Plugins/Experimental/Text3D/Source/Text3D/Private/Renderers/Text3DRendererBase.cpp
@@ -10,8 +10,50 @@
 #include "Extensions/Text3DRenderingExtensionBase.h"
 #include "GameFramework/Actor.h"
 #include "Logs/Text3DLogs.h"
+#include "Settings/Text3DProjectSettings.h"
 #include "Text3DComponent.h"
 
+#if WITH_EDITOR
+void UText3DRendererBase::BindDebugDelegate()
+{
+	if (UText3DProjectSettings* Text3DSettings = UText3DProjectSettings::GetMutable())
+	{
+		Text3DSettings->OnSettingChanged().AddUObject(this, &UText3DRendererBase::OnTextSettingsChanged);
+	}
+}
+
+void UText3DRendererBase::UnbindDebugDelegate()
+{
+	if (UText3DProjectSettings* Text3DSettings = UText3DProjectSettings::GetMutable())
+	{
+		Text3DSettings->OnSettingChanged().RemoveAll(this);
+	}
+}
+
+void UText3DRendererBase::OnTextSettingsChanged(UObject* InSettings, FPropertyChangedEvent& InEvent)
+{
+	if (InEvent.GetMemberPropertyName() == UText3DProjectSettings::GetDebugModePropertyName())
+	{
+		OnDebugModeChanged();
+	}
+}
+
+void UText3DRendererBase::OnDebugModeChanged()
+{
+	if (const UText3DProjectSettings* TextSettings = UText3DProjectSettings::Get())
+	{
+		if (TextSettings->GetDebugMode())
+		{
+			OnDebugModeEnabled();
+		}
+		else
+		{
+			OnDebugModeDisabled();
+		}
+	}
+}
+#endif
+
 UText3DComponent* UText3DRendererBase::GetText3DComponent() const
 {
 	return GetTypedOuter<UText3DComponent>();
@@ -36,6 +78,12 @@ void UText3DRendererBase::Create()
 	}
 
 	OnCreate();
+
+#if WITH_EDITOR
+	BindDebugDelegate();
+	OnDebugModeChanged();
+#endif
+
 	bInitialized = true;
 	UE_LOG(LogText3D, Log, TEXT("%s : Text3DRenderer %s Created"), *Text3DComponent->GetOwner()->GetActorNameOrLabel(), *GetFriendlyName().ToString())
 }
@@ -171,6 +219,11 @@ void UText3DRendererBase::Destroy()
 
 	OnDestroy();
 	CachedBounds.Reset();
+
+#if WITH_EDITOR
+	UnbindDebugDelegate();
+#endif
+
 	bInitialized = false;
 
 	const AActor* Owner = Text3DComponent->GetOwner();
diff --git a/Engine/Plugins/Experimental/Text3D/Source/Text3D/Private/Settings/Text3DProjectSettings.cpp b/Engine/Plugins/Experimental/Text3D/Source/Text3D/Private/Settings/Text3DProjectSettings.cpp
--- a/Engine/Plugins/Experimental/Text3D/Source/Text3D/Private/Settings/Text3DProjectSettings.cpp
+++ b/Engine/Plugins/Experimental/Text3D/Source/Text3D/Private/Settings/Text3DProjectSettings.cpp
@@ -121,6 +121,7 @@ void UText3DProjectSettings::PostEditChangeProperty(struct FPropertyChangedEvent
 
 	static const TSet<FName> PropertyNames
 	{
+		GET_MEMBER_NAME_CHECKED(UText3DProjectSettings, DefaultMaterial),
 		GET_MEMBER_NAME_CHECKED(UText3DProjectSettings, FallbackFont),
 		GET_MEMBER_NAME_CHECKED(UText3DProjectSettings, FallbackFontFace),
 		GET_MEMBER_NAME_CHECKED(UText3DProjectSettings, FontDirectory),
@@ -128,9 +129,12 @@ void UText3DProjectSettings::PostEditChangeProperty(struct FPropertyChangedEvent
 		GET_MEMBER_NAME_CHECKED(UText3DProjectSettings, bShowOnlyMonospaced),
 		GET_MEMBER_NAME_CHECKED(UText3DProjectSettings, bShowOnlyBold),
 		GET_MEMBER_NAME_CHECKED(UText3DProjectSettings, bShowOnlyItalic),
+		GET_MEMBER_NAME_CHECKED(UText3DProjectSettings, bDebugMode),
 	};
 
-	if (PropertyNames.Contains(InChangeEvent.GetMemberPropertyName()))
+	const FName MemberName = InChangeEvent.GetMemberPropertyName();
+
+	if (PropertyNames.Contains(MemberName))
 	{
 		Save();
 	}
@@ -199,6 +203,16 @@ void UText3DProjectSettings::SetShowOnlyItalic(bool bInShowOnlyItalic)
 	}
 }
 
+FName UText3DProjectSettings::GetDebugModePropertyName()
+{
+	return GET_MEMBER_NAME_CHECKED(UText3DProjectSettings, bDebugMode);
+}
+
+bool UText3DProjectSettings::GetDebugMode() const
+{
+	return bDebugMode;
+}
+
 void UText3DProjectSettings::OpenEditorSettingsWindow() const
 {
 	if (ISettingsModule* SettingsModule = FModuleManager::GetModulePtr<ISettingsModule>(TEXT("Settings")))
diff --git a/Engine/Plugins/Experimental/Text3D/Source/Text3D/Public/Renderers/Text3DRendererBase.h b/Engine/Plugins/Experimental/Text3D/Source/Text3D/Public/Renderers/Text3DRendererBase.h
--- a/Engine/Plugins/Experimental/Text3D/Source/Text3D/Public/Renderers/Text3DRendererBase.h
+++ b/Engine/Plugins/Experimental/Text3D/Source/Text3D/Public/Renderers/Text3DRendererBase.h
@@ -52,6 +52,19 @@ protected:
 	/** Calculate the bounds of the rendered text */
 	virtual FBox OnCalculateBounds() const PURE_VIRTUAL(UText3DRenderingImplementationBase::OnCalculateBounds, return FBox(ForceInitToZero); );
 
+#if WITH_EDITOR
+	void BindDebugDelegate();
+	void UnbindDebugDelegate();
+	void OnTextSettingsChanged(UObject* InSettings, FPropertyChangedEvent& InEvent);
+	void OnDebugModeChanged();
+
+	/** Called when entering debug mode in editor */
+	virtual void OnDebugModeEnabled() {}
+
+	/** Called when exiting debug mode in editor */
+	virtual void OnDebugModeDisabled() {}
+#endif
+
 	UText3DComponent* GetText3DComponent() const;
 
 	/** Recalculates bounds when a layout or geometry change happened */
diff --git a/Engine/Plugins/Experimental/Text3D/Source/Text3D/Public/Settings/Text3DProjectSettings.h b/Engine/Plugins/Experimental/Text3D/Source/Text3D/Public/Settings/Text3DProjectSettings.h
--- a/Engine/Plugins/Experimental/Text3D/Source/Text3D/Public/Settings/Text3DProjectSettings.h
+++ b/Engine/Plugins/Experimental/Text3D/Source/Text3D/Public/Settings/Text3DProjectSettings.h
@@ -24,6 +24,10 @@ public:
 	static TEXT3D_API const UText3DProjectSettings* Get();
 	static TEXT3D_API UText3DProjectSettings* GetMutable();
 
+#if WITH_EDITOR
+	static FName GetDebugModePropertyName();
+#endif
+
 	struct FMaterialParameters
 	{
 		static const FLazyName Mode;
@@ -78,6 +82,8 @@ public:
 	TEXT3D_API void SetShowOnlyMonospaced(bool InShowOnlyMonospaced);
 	TEXT3D_API void SetShowOnlyBold(bool InShowOnlyBold);
 	TEXT3D_API void SetShowOnlyItalic(bool InShowOnlyItalic);
+
+	bool GetDebugMode() const;
 #endif
 
 protected:
@@ -127,6 +133,10 @@ protected:
 	/** System font names retrieved from the current platform */
 	UPROPERTY(VisibleAnywhere, Transient, AdvancedDisplay, Category="Text3D")
 	TArray<FString> SystemFontNames;
+
+	/** Debug Text3D components tree hierarchy */
+	UPROPERTY(Config, EditAnywhere, Category="Text3D")
+	bool bDebugMode = false;
 #endif
 
 	/** Base parent material */
