MotionDesign : Text3D
- Added dynamic mesh renderer that uses a single dynamic mesh component to render all glyphs
- Adapted base geometry modifier to handle text mesh updates and retrigger properly
- Clamped scale during layout effect transformation to avoid navigation warning logs about empty bounds

diff --git a/Engine/Plugins/Experimental/Text3D/Source/Text3D/Private/Extensions/Text3DDefaultGeometryExtension.cpp b/Engine/Plugins/Experimental/Text3D/Source/Text3D/Private/Extensions/Text3DDefaultGeometryExtension.cpp
--- a/Engine/Plugins/Experimental/Text3D/Source/Text3D/Private/Extensions/Text3DDefaultGeometryExtension.cpp
+++ b/Engine/Plugins/Experimental/Text3D/Source/Text3D/Private/Extensions/Text3DDefaultGeometryExtension.cpp
@@ -137,13 +137,13 @@ void UText3DDefaultGeometryExtension::PostEditChangeProperty(FPropertyChangedEve
 }
 #endif
 
-UStaticMesh* UText3DDefaultGeometryExtension::FindOrLoadGlyphMesh(UText3DCharacterBase* InCharacter) const
+const FText3DCachedMesh* UText3DDefaultGeometryExtension::FindOrLoadGlyphMesh(UText3DCharacterBase* InCharacter) const
 {
 	if (const FText3DCachedMesh* GlyphMesh = CachedFontData->GetGlyphMesh(InCharacter->GetGlyphIndex(), GlyphMeshParameters))
 	{
 		InCharacter->SetMeshBounds(GlyphMesh->MeshBounds);
 		InCharacter->SetMeshOffset(GlyphMesh->MeshOffset);
-		return GlyphMesh->Mesh;
+		return GlyphMesh;
 	}
 
 	return nullptr;
diff --git a/Engine/Plugins/Experimental/Text3D/Source/Text3D/Private/Extensions/Text3DLayoutTransformEffect.cpp b/Engine/Plugins/Experimental/Text3D/Source/Text3D/Private/Extensions/Text3DLayoutTransformEffect.cpp
--- a/Engine/Plugins/Experimental/Text3D/Source/Text3D/Private/Extensions/Text3DLayoutTransformEffect.cpp
+++ b/Engine/Plugins/Experimental/Text3D/Source/Text3D/Private/Extensions/Text3DLayoutTransformEffect.cpp
@@ -251,7 +251,7 @@ void UText3DLayoutTransformEffect::ApplyEffect(uint32 InGlyphIndex, uint32 InGly
 	{
 		const float Effect = CalculateEffect(InGlyphIndex, InGlyphCount, ScaleOrder, ScaleProgress, ScaleEaseCurve);
 		const FVector Scale = ScaleBegin + (ScaleEnd - ScaleBegin) * Effect;
-		CharacterTransform.MultiplyScale3D(Scale);
+		CharacterTransform.MultiplyScale3D(Scale.ComponentMax(FVector(UE_KINDA_SMALL_NUMBER)));
 	}
 
 	if (bRotationEnabled)
diff --git a/Engine/Plugins/Experimental/Text3D/Source/Text3D/Private/GeometryBuilders/Text3DGlyph.cpp b/Engine/Plugins/Experimental/Text3D/Source/Text3D/Private/GeometryBuilders/Text3DGlyph.cpp
--- a/Engine/Plugins/Experimental/Text3D/Source/Text3D/Private/GeometryBuilders/Text3DGlyph.cpp
+++ b/Engine/Plugins/Experimental/Text3D/Source/Text3D/Private/GeometryBuilders/Text3DGlyph.cpp
@@ -3,11 +3,17 @@
 #include "GeometryBuilders/Text3DGlyph.h"
 
 #include "Engine/StaticMesh.h"
+#include "GeometryScript/GeometryScriptTypes.h"
+#include "GeometryScript/MeshAssetFunctions.h"
 #include "Materials/Material.h"
+#include "Materials/MaterialInterface.h"
+#include "MeshDescription.h"
 #include "Text3DInternalTypes.h"
+#include "UDynamicMesh.h"
+#include "Subsystems/Text3DEngineSubsystem.h"
 
-FText3DGlyph::FText3DGlyph() :
-	StaticMeshAttributes(MeshDescription)
+FText3DGlyph::FText3DGlyph()
+	: StaticMeshAttributes(MeshDescription)
 {
 	StaticMeshAttributes.Register();
 	Groups.SetNum(static_cast<int32>(EText3DGroupType::TypeCount));
@@ -19,11 +25,12 @@ FText3DGlyph::FText3DGlyph() :
 	}
 }
 
-void FText3DGlyph::Build(UStaticMesh* StaticMesh, UMaterial* DefaultMaterial)
+void FText3DGlyph::Build(FText3DCachedMesh& InMesh, UMaterialInterface* InDefaultMaterial)
 {
-	check(StaticMesh);
+	checkf(InMesh.StaticMesh, TEXT("Invalid static mesh object provided to the build function"));
+	checkf(InMesh.DynamicMesh, TEXT("Invalid dynamic mesh object provided to the build function"));
 
-	auto AddMaterial = [](UStaticMesh* InMesh, UMaterial* InMaterial, int32 InIndex)
+	auto AddMaterial = [](UStaticMesh* InMesh, UMaterialInterface* InMaterial, int32 InIndex)
 	{
 		using namespace UE::Text3D::Materials;
 
@@ -47,18 +54,39 @@ void FText3DGlyph::Build(UStaticMesh* StaticMesh, UMaterial* DefaultMaterial)
 		const FPolygonGroupID PolyGroup(Index);
 		if (MeshDescription.GetNumPolygonGroupTriangles(PolyGroup) != 0)
 		{
-			StaticMeshAttributes.GetPolygonGroupMaterialSlotNames()[PolyGroup] = AddMaterial(StaticMesh, DefaultMaterial, Index);
+			StaticMeshAttributes.GetPolygonGroupMaterialSlotNames()[PolyGroup] = AddMaterial(InMesh.StaticMesh, InDefaultMaterial, Index);
 		}
 	}
 
-	TArray<const FMeshDescription*> MeshDescriptions;
-	MeshDescriptions.Add(&MeshDescription);
+	{
+		TArray<const FMeshDescription*> MeshDescriptions;
+		MeshDescriptions.Add(&MeshDescription);
+
+		UStaticMesh::FBuildMeshDescriptionsParams BuildParams;
+		BuildParams.bCommitMeshDescription = true;
+		BuildParams.bFastBuild = true;
+		BuildParams.bAllowCpuAccess = true;
+
+		InMesh.StaticMesh->bAllowCPUAccess = true;
+		InMesh.StaticMesh->BuildFromMeshDescriptions(MeshDescriptions, BuildParams);
+	}
+
+	{
+		static constexpr FGeometryScriptMeshReadLOD StaticMeshLOD
+		{
+			.LODType = EGeometryScriptLODType::SourceModel
+		};
 
-	UStaticMesh::FBuildMeshDescriptionsParams Params;
-	Params.bCommitMeshDescription = true;
-	Params.bFastBuild = true;
-	Params.bAllowCpuAccess = true;
-	StaticMesh->BuildFromMeshDescriptions(MeshDescriptions, Params);
+		static constexpr FGeometryScriptCopyMeshFromAssetOptions ConversionParams
+		{
+			.bApplyBuildSettings = false,
+			.bRequestTangents = false,
+			.bIgnoreRemoveDegenerates = false,
+		};
+
+		EGeometryScriptOutcomePins OutResult;
+		UGeometryScriptLibrary_StaticMeshFunctions::CopyMeshFromStaticMesh(InMesh.StaticMesh, InMesh.DynamicMesh, ConversionParams, StaticMeshLOD, OutResult);
+	}
 }
 
 FMeshDescription& FText3DGlyph::GetMeshDescription()
diff --git a/Engine/Plugins/Experimental/Text3D/Source/Text3D/Private/GeometryBuilders/Text3DGlyph.h b/Engine/Plugins/Experimental/Text3D/Source/Text3D/Private/GeometryBuilders/Text3DGlyph.h
--- a/Engine/Plugins/Experimental/Text3D/Source/Text3D/Private/GeometryBuilders/Text3DGlyph.h
+++ b/Engine/Plugins/Experimental/Text3D/Source/Text3D/Private/GeometryBuilders/Text3DGlyph.h
@@ -8,7 +8,9 @@
 #include "Text3DTypes.h"
 
 class UMaterial;
+class UMaterialInterface;
 class UStaticMesh;
+struct FText3DCachedMesh;
 
 struct FText3DPolygonGroup
 {
@@ -23,7 +25,7 @@ class FText3DGlyph
 public:
 	FText3DGlyph();
 
-	void Build(UStaticMesh* StaticMesh, UMaterial* DefaultMaterial);
+	void Build(FText3DCachedMesh& InMesh, UMaterialInterface* InDefaultMaterial);
 
 	FMeshDescription& GetMeshDescription();
 	FStaticMeshAttributes& GetStaticMeshAttributes();
diff --git a/Engine/Plugins/Experimental/Text3D/Source/Text3D/Private/GeometryBuilders/Text3DGlyphMeshBuilder.cpp b/Engine/Plugins/Experimental/Text3D/Source/Text3D/Private/GeometryBuilders/Text3DGlyphMeshBuilder.cpp
--- a/Engine/Plugins/Experimental/Text3D/Source/Text3D/Private/GeometryBuilders/Text3DGlyphMeshBuilder.cpp
+++ b/Engine/Plugins/Experimental/Text3D/Source/Text3D/Private/GeometryBuilders/Text3DGlyphMeshBuilder.cpp
@@ -7,6 +7,7 @@
 #include "GeometryBuilders/Text3DGlyphContourList.h"
 #include "GeometryBuilders/Text3DGlyphData.h"
 #include "GeometryBuilders/Text3DGlyphPart.h"
+#include "Subsystems/Text3DEngineSubsystem.h"
 
 using namespace UE::Geometry;
 
@@ -128,9 +129,11 @@ void FText3DGlyphMeshBuilder::MovePivot(const FVector& InNewPivot)
 	}
 }
 
-void FText3DGlyphMeshBuilder::BuildMesh(UStaticMesh* StaticMesh, UMaterial* DefaultMaterial)
+void FText3DGlyphMeshBuilder::BuildMesh(FText3DCachedMesh& InMesh, UMaterialInterface* InDefaultMaterial)
 {
-	Glyph->Build(StaticMesh, DefaultMaterial);
+	Glyph->Build(InMesh, InDefaultMaterial);
+	InMesh.MeshBounds = GetMeshBounds();
+	InMesh.MeshOffset = GetMeshOffset();
 }
 
 FBox FText3DGlyphMeshBuilder::GetMeshBounds() const
diff --git a/Engine/Plugins/Experimental/Text3D/Source/Text3D/Private/GeometryBuilders/Text3DGlyphMeshBuilder.h b/Engine/Plugins/Experimental/Text3D/Source/Text3D/Private/GeometryBuilders/Text3DGlyphMeshBuilder.h
--- a/Engine/Plugins/Experimental/Text3D/Source/Text3D/Private/GeometryBuilders/Text3DGlyphMeshBuilder.h
+++ b/Engine/Plugins/Experimental/Text3D/Source/Text3D/Private/GeometryBuilders/Text3DGlyphMeshBuilder.h
@@ -10,7 +10,9 @@ class FText3DGlyphContourList;
 class FText3DGlyphData;
 class FText3DGlyph;
 class UMaterial;
+class UMaterialInterface;
 class UStaticMesh;
+struct FText3DCachedMesh;
 
 /** Creates the glyph mesh from glyph data */
 class FText3DGlyphMeshBuilder final
@@ -35,7 +37,8 @@ public:
 	/** Move the mesh pivot based on pivot ratio 0-1 */
 	void MovePivot(const FVector& InNewPivot);
 
-	void BuildMesh(UStaticMesh* StaticMesh, UMaterial* DefaultMaterial);
+	/** Creates the actual mesh */
+	void BuildMesh(FText3DCachedMesh& InMesh, UMaterialInterface* InDefaultMaterial);
 
 	/** Get actual glyph bounds once built */
 	FBox GetMeshBounds() const;
diff --git a/Engine/Plugins/Experimental/Text3D/Source/Text3D/Private/Renderers/DynamicMesh/Text3DDynamicMeshRenderer.cpp b/Engine/Plugins/Experimental/Text3D/Source/Text3D/Private/Renderers/DynamicMesh/Text3DDynamicMeshRenderer.cpp
new file mode 100644
--- /dev/null
+++ b/Engine/Plugins/Experimental/Text3D/Source/Text3D/Private/Renderers/DynamicMesh/Text3DDynamicMeshRenderer.cpp
@@ -0,0 +1,436 @@
+// Copyright Epic Games, Inc. All Rights Reserved.
+
+#include "Renderers/DynamicMesh/Text3DDynamicMeshRenderer.h"
+
+#include "Characters/Text3DCharacterBase.h"
+#include "Components/DynamicMeshComponent.h"
+#include "DynamicMesh/MeshTransforms.h"
+#include "DynamicMeshEditor.h"
+#include "Extensions/Text3DGeometryExtensionBase.h"
+#include "Extensions/Text3DLayoutExtensionBase.h"
+#include "Extensions/Text3DMaterialExtensionBase.h"
+#include "Extensions/Text3DRenderingExtensionBase.h"
+#include "Subsystems/Text3DEngineSubsystem.h"
+#include "Text3DComponent.h"
+#include "Text3DInternalTypes.h"
+
+UText3DDynamicMeshRenderer::FScopedMeshEditor::~FScopedMeshEditor()
+{
+	if (bUpdated && IsValid(MeshComponent))
+	{
+		MeshComponent->NotifyMeshVertexAttributesModified();
+		MeshComponent->OnMeshChanged.Broadcast();
+	}
+}
+
+void UText3DDynamicMeshRenderer::FScopedMeshEditor::EditMesh(const TFunctionRef<void(UE::Geometry::FDynamicMesh3&)>& InFunctor)
+{
+	if (MeshComponent && MeshComponent->GetDynamicMesh())
+	{
+		constexpr bool bDeferChanges = true;
+		MeshComponent->GetDynamicMesh()->EditMesh(InFunctor, EDynamicMeshChangeType::GeneralEdit, EDynamicMeshAttributeChangeFlags::Unknown, bDeferChanges);
+		bUpdated = true;
+	}
+}
+
+void UText3DDynamicMeshRenderer::FScopedMeshEditor::ProcessMesh(const TFunctionRef<void(const UE::Geometry::FDynamicMesh3&)>& InFunctor) const
+{
+	if (MeshComponent && MeshComponent->GetDynamicMesh())
+	{
+		MeshComponent->GetDynamicMesh()->ProcessMesh(InFunctor);
+	}
+}
+
+void UText3DDynamicMeshRenderer::OnCreate()
+{
+	if (!DynamicMeshComponent)
+	{
+		UText3DComponent* TextComponent = GetText3DComponent();
+		DynamicMeshComponent = NewObject<UDynamicMeshComponent>(this, TEXT("Text3DDynamicMesh"), RF_Transient);
+		DynamicMeshComponent->SetupAttachment(TextComponent);
+		ClearMesh();
+	}
+}
+
+void UText3DDynamicMeshRenderer::OnUpdate(const UE::Text3D::Renderer::FUpdateParameters& InParameters)
+{
+	const UText3DComponent* TextComponent = GetText3DComponent();
+
+	if (!ScopedMeshEditor.IsSet())
+	{
+		ScopedMeshEditor = FScopedMeshEditor(DynamicMeshComponent);
+
+		// Restore checkpoint
+		RestoreMesh(ScopedMeshEditor.GetValue());
+	}
+
+	FScopedMeshEditor& MeshEditor = ScopedMeshEditor.GetValue();
+
+	if (InParameters.CurrentFlag == EText3DRendererFlags::Geometry)
+	{
+		AllocateGlyphMeshData(MeshEditor, TextComponent->GetCharacterCount());
+
+		const UText3DGeometryExtensionBase* GeometryExtension = TextComponent->GetGeometryExtension();
+		TextComponent->ForEachCharacter([this, GeometryExtension, &MeshEditor](UText3DCharacterBase* InCharacter, uint16 InIndex, uint16 InCount)
+		{
+			const FText3DCachedMesh* CachedGlyph = GeometryExtension->FindOrLoadGlyphMesh(InCharacter);
+			AppendGlyphMesh(MeshEditor, InIndex, InCharacter->GetGlyphIndex(), CachedGlyph->DynamicMesh);
+		});
+
+		RefreshBounds();
+	}
+	else if (InParameters.CurrentFlag == EText3DRendererFlags::Layout)
+	{
+		const UText3DLayoutExtensionBase* LayoutExtension = TextComponent->GetLayoutExtension();
+
+		DynamicMeshComponent->SetRelativeScale3D(LayoutExtension->GetTextScale());
+
+		TextComponent->ForEachCharacter([this, &MeshEditor](UText3DCharacterBase* InCharacter, uint16 InIndex, uint16 InCount)
+		{
+			constexpr bool bReset = false;
+			const FTransform& CharacterTransform = InCharacter->GetTransform(bReset);
+			SetGlyphMeshTransform(MeshEditor, InIndex, CharacterTransform, InCharacter->GetVisibility());
+		});
+
+		RefreshBounds();
+	}
+	else if (InParameters.CurrentFlag == EText3DRendererFlags::Material)
+	{
+		using namespace UE::Text3D::Materials;
+
+		const UText3DMaterialExtensionBase* MaterialExtension = TextComponent->GetMaterialExtension();
+
+		TextComponent->ForEachCharacter([this, MaterialExtension](UText3DCharacterBase* InCharacter, uint16 InIndex, uint16 InCount)
+		{
+			for (const int32 MaterialSlot : GlyphMeshData[InIndex].MaterialSlots)
+			{
+				UMaterialInterface* Material = MaterialExtension->GetMaterial(static_cast<EText3DGroupType>(MaterialSlot - (InIndex * SlotNames.Num())));
+
+				if (Material != DynamicMeshComponent->GetMaterial(MaterialSlot))
+				{
+					DynamicMeshComponent->SetMaterial(MaterialSlot, Material);
+				}
+			}
+		});
+	}
+	else if (InParameters.CurrentFlag == EText3DRendererFlags::Visibility)
+	{
+		const UText3DRenderingExtensionBase* RenderingExtension = TextComponent->GetRenderingExtension();
+
+		DynamicMeshComponent->SetHiddenInGame(TextComponent->bHiddenInGame);
+		DynamicMeshComponent->SetCastShadow(RenderingExtension->GetTextCastShadow());
+		DynamicMeshComponent->SetCastHiddenShadow(RenderingExtension->GetTextCastHiddenShadow());
+		DynamicMeshComponent->SetAffectDynamicIndirectLighting(RenderingExtension->GetTextAffectDynamicIndirectLighting());
+		DynamicMeshComponent->SetAffectIndirectLightingWhileHidden(RenderingExtension->GetTextAffectIndirectLightingWhileHidden());
+		DynamicMeshComponent->SetHoldout(RenderingExtension->GetTextHoldout());
+		DynamicMeshComponent->SetVisibility(TextComponent->GetVisibleFlag());
+
+		TextComponent->ForEachCharacter([this, &MeshEditor](UText3DCharacterBase* InCharacter, uint16 InIndex, uint16 InCount)
+		{
+			constexpr bool bReset = false;
+			const FTransform& CharacterTransform = InCharacter->GetTransform(bReset);
+			SetGlyphMeshTransform(MeshEditor, InIndex, CharacterTransform, InCharacter->GetVisibility());
+		});
+	}
+
+	if (InParameters.bIsLastFlag)
+	{
+		// Save checkpoint
+		SaveMesh(MeshEditor);
+
+		ScopedMeshEditor.Reset();
+	}
+}
+
+void UText3DDynamicMeshRenderer::OnClear()
+{
+	if (DynamicMeshComponent)
+	{
+		DynamicMeshComponent->EditMesh([](FDynamicMesh3& InEditMesh)
+		{
+			for (const int32 TId : InEditMesh.TriangleIndicesItr())
+			{
+				InEditMesh.RemoveTriangle(TId);
+			}
+		});
+	}
+}
+
+void UText3DDynamicMeshRenderer::OnDestroy()
+{
+	if (DynamicMeshComponent)
+	{
+		DynamicMeshComponent->DetachFromComponent(FDetachmentTransformRules::KeepRelativeTransform);
+		DynamicMeshComponent->DestroyComponent();
+		DynamicMeshComponent = nullptr;
+	}
+}
+
+FName UText3DDynamicMeshRenderer::GetFriendlyName() const
+{
+	static const FName Name(TEXT("DynamicMeshRenderer"));
+	return Name;
+}
+
+FBox UText3DDynamicMeshRenderer::OnCalculateBounds() const
+{
+	FBox Box(ForceInit);
+	Box += DynamicMeshComponent->Bounds.GetBox();
+	return Box;
+}
+
+#if WITH_EDITOR
+void UText3DDynamicMeshRenderer::OnDebugModeEnabled()
+{
+	Super::OnDebugModeEnabled();
+	SetDebugMode(true);
+}
+
+void UText3DDynamicMeshRenderer::OnDebugModeDisabled()
+{
+	Super::OnDebugModeDisabled();
+	SetDebugMode(false);
+}
+
+void UText3DDynamicMeshRenderer::SetDebugMode(bool bInEnabled)
+{
+	// Since we are dealing with class FProperty, no need to run this for each instance, do it once
+	static bool bDebugModeEnabled = true;
+
+	if (bDebugModeEnabled != bInEnabled)
+	{
+		bDebugModeEnabled = bInEnabled;
+
+		FProperty* DynamicMeshComponentProperty = FindFProperty<FProperty>(UText3DDynamicMeshRenderer::StaticClass(), GET_MEMBER_NAME_CHECKED(UText3DDynamicMeshRenderer, DynamicMeshComponent));
+
+		// Here we toggle the CPF_Edit flag to hide/show property in the details panel component editor tree / outliner
+		// @see FComponentEditorUtils::GetPropertyForEditableNativeComponent
+		// todo : implement a custom debug view widget for text or add a property editor metadata to control the component visibility in the component editor tree / outliner
+		if (bInEnabled)
+		{
+			DynamicMeshComponentProperty->SetPropertyFlags(CPF_Edit);
+		}
+		else
+		{
+			DynamicMeshComponentProperty->ClearPropertyFlags(CPF_Edit);
+		}
+	}
+}
+#endif
+
+void UText3DDynamicMeshRenderer::AllocateGlyphMeshData(FScopedMeshEditor& InMeshEditor, int32 InCount)
+{
+	if (GlyphMeshData.Num() > InCount)
+	{
+		for (int32 Index = InCount; Index < GlyphMeshData.Num(); Index++)
+		{
+			ClearGlyphMesh(InMeshEditor, Index);
+		}
+	}
+
+	GlyphMeshData.SetNum(InCount);
+}
+
+void UText3DDynamicMeshRenderer::AppendGlyphMesh(FScopedMeshEditor& InMeshEditor, uint16 InIndex, uint32 InGlyphIndex, UDynamicMesh* InMesh)
+{
+	using namespace UE::Geometry;
+
+	check(!!InMesh);
+
+	if (!GlyphMeshData.IsValidIndex(InIndex))
+	{
+		return;
+	}
+
+	FBox Bounds = FBox(ForceInitToZero);
+	InMesh->ProcessMesh([&Bounds](const FDynamicMesh3& InMesh)
+	{
+		Bounds = static_cast<FBox>(InMesh.GetBounds(true));
+	});
+
+	if (GlyphMeshData[InIndex].GlyphIndex == InGlyphIndex && GlyphMeshData[InIndex].Bounds.Equals(Bounds))
+	{
+		return;
+	}
+
+	ClearGlyphMesh(InMeshEditor, InIndex);
+
+	GlyphMeshData[InIndex].GlyphIndex = InGlyphIndex;
+	GlyphMeshData[InIndex].Bounds = Bounds;
+	GlyphMeshData[InIndex].CurrentTransform = FTransform::Identity;
+
+	InMeshEditor.EditMesh(
+		[this, InMesh, InIndex, InGlyphIndex](FDynamicMesh3& InEditMesh)
+		{
+			InMesh->ProcessMesh([this, &InEditMesh, InIndex, InGlyphIndex](const FDynamicMesh3& InProcessMesh)
+			{
+				FMeshIndexMappings Mappings;
+				FDynamicMeshEditor Editor(&InEditMesh);
+				Editor.AppendMesh(&InProcessMesh, Mappings);
+
+				TArray<int32> NewTriangles;
+				const TMap<int32, int32>& FromToMap = Mappings.GetTriangleMap().GetForwardMap();
+				FDynamicMeshMaterialAttribute* MaterialAttribute = InEditMesh.Attributes()->GetMaterialID();
+
+				TArray<int32> MaterialSlots;
+				for (const TPair<int32, int32>& FromToPair : FromToMap)
+				{
+					InEditMesh.SetTriangleGroup(FromToPair.Value, InIndex);
+					const int32 MaterialID = (InIndex * UE::Text3D::Materials::SlotNames.Num()) + MaterialAttribute->GetValue(FromToPair.Value);
+					MaterialAttribute->SetValue(FromToPair.Value, MaterialID);
+					MaterialSlots.AddUnique(MaterialID);
+				}
+
+				MaterialSlots.StableSort([](int32 A, int32 B)
+				{
+					return A > B;
+				});
+
+				GlyphMeshData[InIndex].MaterialSlots = MoveTemp(MaterialSlots);
+			});
+		});
+}
+
+void UText3DDynamicMeshRenderer::ClearMesh()
+{
+	GlyphMeshData.Reset();
+	DynamicMeshComponent->EditMesh([](FDynamicMesh3& InEditMesh)
+	{
+		InEditMesh.Clear();
+		InEditMesh.EnableTriangleGroups();
+		InEditMesh.EnableAttributes();
+		InEditMesh.Attributes()->EnableMaterialID();
+		InEditMesh.Attributes()->EnableTangents();
+		InEditMesh.Attributes()->SetNumUVLayers(1);
+	});
+}
+
+void UText3DDynamicMeshRenderer::ClearGlyphMesh(FScopedMeshEditor& InMeshEditor, uint16 InIndex)
+{
+	using namespace UE::Geometry;
+
+	// Remove old glyph data
+	InMeshEditor.EditMesh(
+		[this, InIndex](FDynamicMesh3& InEditMesh)
+		{
+			TArray<int32> TrianglesToDelete;
+			for (const int32 Tid : InEditMesh.TriangleIndicesItr())
+			{
+				if (InEditMesh.GetTriangleGroup(Tid) == InIndex)
+				{
+					TrianglesToDelete.Add(Tid);
+				}
+			}
+
+			FDynamicMeshEditor Editor(&InEditMesh);
+			Editor.RemoveTriangles(TrianglesToDelete, /** RemoveIsolatedVertices*/ true);
+			InEditMesh.CompactInPlace();
+		});
+}
+
+void UText3DDynamicMeshRenderer::SetGlyphMeshTransform(FScopedMeshEditor& InMeshEditor, uint16 InIndex, const FTransform& InTransform, bool bInVisible)
+{
+	using namespace UE::Geometry;
+
+	if (!GlyphMeshData.IsValidIndex(InIndex))
+	{
+		return;
+	}
+
+	FGlyphMeshData& GlyphData = GlyphMeshData[InIndex];
+
+	if (!GlyphData.bVisible && GlyphData.bVisible == bInVisible)
+	{
+		return;
+	}
+
+	GlyphData.bVisible = bInVisible;
+
+	const FVector TargetScale = GlyphData.bVisible ? InTransform.GetScale3D().ComponentMax(FVector(UE_KINDA_SMALL_NUMBER)) : FVector(UE_KINDA_SMALL_NUMBER);
+	const FTransform TargetTransform(InTransform.GetRotation(), InTransform.GetTranslation(), TargetScale);
+
+	if (GlyphData.CurrentTransform.Equals(TargetTransform))
+	{
+		return;
+	}
+
+	InMeshEditor.EditMesh([InIndex, &GlyphData, &TargetTransform](FDynamicMesh3& InEditMesh)
+	{
+		FDynamicMeshNormalOverlay* Normals = InEditMesh.Attributes()->PrimaryNormals();
+		FDynamicMeshNormalOverlay* Tangents = InEditMesh.Attributes()->PrimaryTangents();
+
+		TSet<int32> VerticesToTransform;
+		TSet<int32> NormalsToTransform;
+		TSet<int32> TangentsToTransform;
+		for (const int32 TId : InEditMesh.TriangleIndicesItr())
+		{
+			if (InEditMesh.GetTriangleGroup(TId) == InIndex)
+			{
+				FIndex3i Tri = InEditMesh.GetTriangle(TId);
+
+				VerticesToTransform.Add(Tri.A);
+				VerticesToTransform.Add(Tri.B);
+				VerticesToTransform.Add(Tri.C);
+
+				NormalsToTransform.Add(Normals->GetElementIDAtVertex(TId, Tri.A));
+				NormalsToTransform.Add(Normals->GetElementIDAtVertex(TId, Tri.B));
+				NormalsToTransform.Add(Normals->GetElementIDAtVertex(TId, Tri.C));
+
+				TangentsToTransform.Add(Tangents->GetElementIDAtVertex(TId, Tri.A));
+				TangentsToTransform.Add(Tangents->GetElementIDAtVertex(TId, Tri.B));
+				TangentsToTransform.Add(Tangents->GetElementIDAtVertex(TId, Tri.C));
+			}
+		}
+
+		const FTransform3d Inverse = GlyphData.CurrentTransform.Inverse();
+		for (const int32 VId : VerticesToTransform)
+		{
+			FVector3d LocalPos = Inverse.TransformPosition(InEditMesh.GetVertex(VId));
+			FVector3d NewWorldPos = TargetTransform.TransformPosition(LocalPos);
+			InEditMesh.SetVertex(VId, NewWorldPos);
+		}
+
+		const FQuat DeltaRotation = TargetTransform.GetRotation() * GlyphData.CurrentTransform.GetRotation().Inverse();
+		for (const int32 NId : NormalsToTransform)
+		{
+			if (NId != FDynamicMesh3::InvalidID)
+			{
+				FVector3f Normal = Normals->GetElement(NId);
+				FVector3d NewNormal = DeltaRotation.RotateVector((FVector3d)Normal);
+				Normals->SetElement(NId, (FVector3f)NewNormal.GetSafeNormal());
+			}
+		}
+
+		for (const int32 TId : TangentsToTransform)
+		{
+			if (TId != FDynamicMesh3::InvalidID)
+			{
+				FVector3f Tangent = Tangents->GetElement(TId);
+				FVector3d NewTangent = DeltaRotation.RotateVector((FVector3d)Tangent);
+				Tangents->SetElement(TId, (FVector3f)NewTangent.GetSafeNormal());	
+			}
+		}
+	});
+
+	GlyphData.CurrentTransform = TargetTransform;
+}
+
+void UText3DDynamicMeshRenderer::SaveMesh(const FScopedMeshEditor& InMeshEditor)
+{
+	InMeshEditor.ProcessMesh([this](const FDynamicMesh3& EditMesh)
+	{
+		CachedMesh = EditMesh;
+	});
+}
+
+void UText3DDynamicMeshRenderer::RestoreMesh(FScopedMeshEditor& InMeshEditor)
+{
+	if (CachedMesh.IsSet())
+	{
+		InMeshEditor.EditMesh([this](FDynamicMesh3& EditMesh)
+		{
+			EditMesh = MoveTemp(CachedMesh.GetValue());
+			CachedMesh.Reset();
+		});
+	}
+}
diff --git a/Engine/Plugins/Experimental/Text3D/Source/Text3D/Private/Renderers/DynamicMesh/Text3DDynamicMeshRenderer.h b/Engine/Plugins/Experimental/Text3D/Source/Text3D/Private/Renderers/DynamicMesh/Text3DDynamicMeshRenderer.h
new file mode 100644
--- /dev/null
+++ b/Engine/Plugins/Experimental/Text3D/Source/Text3D/Private/Renderers/DynamicMesh/Text3DDynamicMeshRenderer.h
@@ -0,0 +1,82 @@
+// Copyright Epic Games, Inc. All Rights Reserved.
+
+#pragma once
+
+#include "DynamicMesh/DynamicMesh3.h"
+#include "Renderers/Text3DRendererBase.h"
+#include "UObject/ObjectPtr.h"
+#include "Text3DDynamicMeshRenderer.generated.h"
+
+class UDynamicMesh;
+class UDynamicMeshComponent;
+
+/**
+ * Renderer that uses a single dynamic mesh component to render all characters
+ * Text3DComponent
+ * - Text3DDynamicMesh
+ */
+UCLASS(DisplayName="DynamicMeshRenderer", ClassGroup="Text3D", HideCategories=(Renderer))
+class UText3DDynamicMeshRenderer : public UText3DRendererBase
+{
+	GENERATED_BODY()
+
+protected:
+	/** Used to store data about character glyph mesh */
+	struct FGlyphMeshData
+	{
+		bool bVisible = true;
+		uint16 GlyphIndex = 0;
+		TArray<int32> MaterialSlots;
+		FTransform CurrentTransform = FTransform::Identity;
+		FBox Bounds = FBox(ForceInitToZero);
+	};
+
+	/** Used to push render changes when destroyed */
+	struct FScopedMeshEditor
+	{
+		explicit FScopedMeshEditor(UDynamicMeshComponent* InComponent)
+			: MeshComponent(InComponent)
+		{}
+
+		~FScopedMeshEditor();
+		void EditMesh(const TFunctionRef<void(UE::Geometry::FDynamicMesh3&)>& InFunctor);
+		void ProcessMesh(const TFunctionRef<void(const UE::Geometry::FDynamicMesh3&)>& InFunctor) const;
+
+	private:
+		UDynamicMeshComponent* MeshComponent = nullptr;
+		bool bUpdated = false;
+	};
+
+	//~ Begin UText3DRendererBase
+	virtual void OnCreate() override;
+	virtual void OnUpdate(const UE::Text3D::Renderer::FUpdateParameters& InParameters) override;
+	virtual void OnClear() override;
+	virtual void OnDestroy() override;
+	virtual FName GetFriendlyName() const override;
+	virtual FBox OnCalculateBounds() const override;
+#if WITH_EDITOR
+	virtual void OnDebugModeEnabled() override;
+	virtual void OnDebugModeDisabled() override;
+#endif
+	//~ End UText3DRendererBase
+
+#if WITH_EDITOR
+	void SetDebugMode(bool bInEnabled);
+#endif
+
+	void AllocateGlyphMeshData(FScopedMeshEditor& InMeshEditor, int32 InCount);
+	void AppendGlyphMesh(FScopedMeshEditor& InMeshEditor, uint16 InIndex, uint32 InGlyphIndex, UDynamicMesh* InMesh);
+	void ClearMesh();
+	void ClearGlyphMesh(FScopedMeshEditor& InMeshEditor, uint16 InIndex);
+	void SetGlyphMeshTransform(FScopedMeshEditor& InMeshEditor, uint16 InIndex, const FTransform& InTransform, bool bInVisible);
+	
+	void SaveMesh(const FScopedMeshEditor& InMeshEditor);
+	void RestoreMesh(FScopedMeshEditor& InMeshEditor);
+
+	UPROPERTY(VisibleAnywhere, Category="Renderer", meta=(EditCondition="false", EditConditionHides))
+	TObjectPtr<UDynamicMeshComponent> DynamicMeshComponent;
+
+	TArray<FGlyphMeshData> GlyphMeshData;
+	TOptional<FScopedMeshEditor> ScopedMeshEditor;
+	TOptional<UE::Geometry::FDynamicMesh3> CachedMesh;
+};
diff --git a/Engine/Plugins/Experimental/Text3D/Source/Text3D/Private/Renderers/StaticMeshes/Text3DStaticMeshesRenderer.cpp b/Engine/Plugins/Experimental/Text3D/Source/Text3D/Private/Renderers/StaticMeshes/Text3DStaticMeshesRenderer.cpp
--- a/Engine/Plugins/Experimental/Text3D/Source/Text3D/Private/Renderers/StaticMeshes/Text3DStaticMeshesRenderer.cpp
+++ b/Engine/Plugins/Experimental/Text3D/Source/Text3D/Private/Renderers/StaticMeshes/Text3DStaticMeshesRenderer.cpp
@@ -12,6 +12,7 @@
 #include "Extensions/Text3DMaterialExtensionBase.h"
 #include "Extensions/Text3DRenderingExtensionBase.h"
 #include "Logs/Text3DLogs.h"
+#include "Subsystems/Text3DEngineSubsystem.h"
 #include "Text3DComponent.h"
 #include "Text3DInternalTypes.h"
 #include "Utilities/Text3DUtilities.h"
@@ -49,10 +50,10 @@ void UText3DStaticMeshesRenderer::OnUpdate(const UE::Text3D::Renderer::FUpdatePa
 
 		TextComponent->ForEachCharacter([this, GeometryExtension](UText3DCharacterBase* InCharacter, uint16 InIndex, uint16 InCount)
 		{
-			UStaticMesh* CachedMesh = GeometryExtension->FindOrLoadGlyphMesh(InCharacter);
+			const FText3DCachedMesh* CachedMesh = GeometryExtension->FindOrLoadGlyphMesh(InCharacter);
 			if (CharacterMeshes.IsValidIndex(InIndex))
 			{
-				CharacterMeshes[InIndex]->SetStaticMesh(CachedMesh);
+				CharacterMeshes[InIndex]->SetStaticMesh(CachedMesh->StaticMesh);
 			}
 			else
 			{
diff --git a/Engine/Plugins/Experimental/Text3D/Source/Text3D/Private/Renderers/StaticMeshes/Text3DStaticMeshesRenderer.h b/Engine/Plugins/Experimental/Text3D/Source/Text3D/Private/Renderers/StaticMeshes/Text3DStaticMeshesRenderer.h
--- a/Engine/Plugins/Experimental/Text3D/Source/Text3D/Private/Renderers/StaticMeshes/Text3DStaticMeshesRenderer.h
+++ b/Engine/Plugins/Experimental/Text3D/Source/Text3D/Private/Renderers/StaticMeshes/Text3DStaticMeshesRenderer.h
@@ -14,12 +14,11 @@ class UText3DComponent;
 /**
  * Legacy/default renderer for Text3D
  * Each text character is rendered as a StaticMesh within its own StaticMeshComponent
- * Kerning is done through a scene component containing the current character
  * Text3DComponent
  * - Text3DRoot (Root)
  * --- Text3DStaticMeshComponent (Character)
  *
- * Eg: The text "Hello" will be rendered using 5 SceneComponents and 5 StaticMeshComponents
+ * Eg: The text "Hello" will be rendered using 5 StaticMeshComponents
  */
 UCLASS(DisplayName="StaticMeshesRenderer", ClassGroup="Text3D", HideCategories=(Renderer))
 class UText3DStaticMeshesRenderer : public UText3DRendererBase
diff --git a/Engine/Plugins/Experimental/Text3D/Source/Text3D/Private/Renderers/Text3DRendererBase.cpp b/Engine/Plugins/Experimental/Text3D/Source/Text3D/Private/Renderers/Text3DRendererBase.cpp
--- a/Engine/Plugins/Experimental/Text3D/Source/Text3D/Private/Renderers/Text3DRendererBase.cpp
+++ b/Engine/Plugins/Experimental/Text3D/Source/Text3D/Private/Renderers/Text3DRendererBase.cpp
@@ -126,19 +126,23 @@ void UText3DRendererBase::Update(EText3DRendererFlags InFlags)
 	});
 
 	uint8 Flag = 1 << 0;
+	uint8 FlagsLeft = static_cast<uint8>(InFlags);
 
 	while (Flag < static_cast<uint8>(EText3DRendererFlags::All))
 	{
 		const EText3DRendererFlags FlagCasted = static_cast<EText3DRendererFlags>(Flag);
 
-		const FUpdateParameters Parameters
-		{
-			.UpdateFlags = InFlags,
-			.CurrentFlag = FlagCasted
-		};
-
 		if (EnumHasAnyFlags(InFlags, FlagCasted))
 		{
+			FlagsLeft -= Flag;
+
+			const FUpdateParameters Parameters
+			{
+				.UpdateFlags = InFlags,
+				.CurrentFlag = FlagCasted,
+				.bIsLastFlag = FlagsLeft == 0
+			};
+
 			for (int32 Index = Extensions.Num() - 1; Index >= 0; --Index)
 			{
 				if (UText3DExtensionBase* Extension = Extensions[Index])
@@ -238,4 +242,4 @@ FBox UText3DRendererBase::GetBounds() const
 	}
 
 	return CachedBounds.GetValue();
-}
+}
\ No newline at end of file
diff --git a/Engine/Plugins/Experimental/Text3D/Source/Text3D/Private/Subsystems/Text3DEngineSubsystem.cpp b/Engine/Plugins/Experimental/Text3D/Source/Text3D/Private/Subsystems/Text3DEngineSubsystem.cpp
--- a/Engine/Plugins/Experimental/Text3D/Source/Text3D/Private/Subsystems/Text3DEngineSubsystem.cpp
+++ b/Engine/Plugins/Experimental/Text3D/Source/Text3D/Private/Subsystems/Text3DEngineSubsystem.cpp
@@ -1,7 +1,9 @@
 // Copyright Epic Games, Inc. All Rights Reserved.
 
 #include "Text3DEngineSubsystem.h"
+
 #include "Containers/Ticker.h"
+#include "UDynamicMesh.h"
 #include "Engine/Font.h"
 #include "Engine/Engine.h"
 #include "Engine/StaticMesh.h"
@@ -315,7 +317,7 @@ FText3DCachedMesh* FCachedFontData::GetGlyphMesh(uint32 InGlyphIndex, const FGly
 	uint32 HashGroup = 0;
 	HashGroup = HashCombine(HashGroup, GetTypeHash(Font));
 	HashGroup = HashCombine(HashGroup, GetTypeHash(InGlyphIndex));
-	const FString StaticMeshName = FString::Printf(TEXT("Text3D_Char_%u_%u"), HashGroup, HashParameters);
+	const FString MeshName = FString::Printf(TEXT("Text3D_Char_%u_%u"), HashGroup, HashParameters);
 
 	const TText3DGlyphContourNodeShared Root = GetGlyphContours(InGlyphIndex, TypefaceIndex);
 	if (!Root || Root->Children.IsEmpty())
@@ -323,22 +325,26 @@ FText3DCachedMesh* FCachedFontData::GetGlyphMesh(uint32 InGlyphIndex, const FGly
 		return nullptr;
 	}
 
-	UText3DEngineSubsystem* Subsystem = UText3DEngineSubsystem::Get();
-	const FName StrokeShapeUniqueName = MakeUniqueObjectName(Subsystem, UStaticMesh::StaticClass(), FName(*StaticMeshName));
+	UText3DEngineSubsystem* TextSubsystem = UText3DEngineSubsystem::Get();
+	if (!TextSubsystem)
+	{
+		return nullptr;
+	}
 
 	FText3DCachedMesh& CachedMesh = CachedMeshes.Glyphs.FindOrAdd(InGlyphIndex);
-	CachedMesh.Mesh = NewObject<UStaticMesh>(Subsystem, StrokeShapeUniqueName, RF_Transient | RF_Public);
+
+	const FName StaticMeshName = MakeUniqueObjectName(TextSubsystem, UStaticMesh::StaticClass(), FName(*MeshName));
+	CachedMesh.StaticMesh = NewObject<UStaticMesh>(TextSubsystem, StaticMeshName, RF_Transient | RF_Public);
+
+	const FName DynamicMeshName = MakeUniqueObjectName(TextSubsystem, UDynamicMesh::StaticClass(), FName(*MeshName));
+	CachedMesh.DynamicMesh = NewObject<UDynamicMesh>(TextSubsystem, DynamicMeshName, RF_Transient | RF_Public);
 
 	FText3DGlyphMeshBuilder MeshCreator;
 	MeshCreator.CreateMeshes(Root, InParameters.Extrude, InParameters.Bevel, InParameters.BevelType, InParameters.BevelSegments, InParameters.bOutline, InParameters.OutlineExpand);
 	MeshCreator.SetFrontAndBevelTextureCoordinates(InParameters.Bevel);
 	MeshCreator.MirrorGroups(InParameters.Extrude);
 	MeshCreator.MovePivot(InParameters.PivotOffset);
-
-	MeshCreator.BuildMesh(CachedMesh.Mesh, Subsystem->DefaultMaterial);
-
-	CachedMesh.MeshBounds = MeshCreator.GetMeshBounds();
-	CachedMesh.MeshOffset = MeshCreator.GetMeshOffset();
+	MeshCreator.BuildMesh(CachedMesh, TextSubsystem->DefaultMaterial);
 
 	return &CachedMesh;
 }
diff --git a/Engine/Plugins/Experimental/Text3D/Source/Text3D/Private/Subsystems/Text3DEngineSubsystem.h b/Engine/Plugins/Experimental/Text3D/Source/Text3D/Private/Subsystems/Text3DEngineSubsystem.h
--- a/Engine/Plugins/Experimental/Text3D/Source/Text3D/Private/Subsystems/Text3DEngineSubsystem.h
+++ b/Engine/Plugins/Experimental/Text3D/Source/Text3D/Private/Subsystems/Text3DEngineSubsystem.h
@@ -11,6 +11,7 @@
 #include "Text3DEngineSubsystem.generated.h"
 
 class FFreeTypeFace;
+class UDynamicMesh;
 class UFont;
 class UMaterial;
 class UStaticMesh;
@@ -67,7 +68,10 @@ struct FText3DCachedMesh
 	GENERATED_BODY()
 
 	UPROPERTY()
-	TObjectPtr<UStaticMesh> Mesh;
+	TObjectPtr<UStaticMesh> StaticMesh;
+
+	UPROPERTY()
+	TObjectPtr<UDynamicMesh> DynamicMesh;
 
 	UPROPERTY()
 	FBox MeshBounds = FBox(ForceInitToZero);
diff --git a/Engine/Plugins/Experimental/Text3D/Source/Text3D/Private/Text3DComponent.cpp b/Engine/Plugins/Experimental/Text3D/Source/Text3D/Private/Text3DComponent.cpp
--- a/Engine/Plugins/Experimental/Text3D/Source/Text3D/Private/Text3DComponent.cpp
+++ b/Engine/Plugins/Experimental/Text3D/Source/Text3D/Private/Text3DComponent.cpp
@@ -344,6 +344,23 @@ void UText3DComponent::PostEditImport()
 {
 	Super::PostEditImport();
 
+	// Clear renderer after duplication to avoid duplicated characters
+	TArray<UObject*> OwnedObjects;
+	GetObjectsWithOuter(this, OwnedObjects, false);
+
+	for (UObject* OwnedObject : OwnedObjects)
+	{
+		if (UText3DRendererBase* Renderer = Cast<UText3DRendererBase>(OwnedObject))
+		{
+			if (Renderer != TextRenderer)
+			{
+				Renderer->Clear();
+				Renderer->Destroy();
+				Renderer->MarkAsGarbage();
+			}
+		}
+	}
+
 	RequestUpdate(EText3DRendererFlags::All);
 }
 
@@ -413,6 +430,13 @@ void UText3DComponent::SetFontByName(const FString& InFontName)
 		}
 	}
 }
+
+void UText3DComponent::ForceUpdateText()
+{
+	const FText PrevText = GetText();
+	SetText(FText::GetEmpty());
+	SetText(PrevText);
+}
 #endif
 
 const FText& UText3DComponent::GetFormattedText() const
@@ -982,7 +1006,7 @@ void UText3DComponent::Rebuild(EText3DRendererFlags InUpdateFlags, bool bCleanCa
 
 	if (!bInitialized)
 	{
-		UE_LOG(LogText3D, Warning, TEXT("[%s] Rebuild called with update flags (%u) while text is not initialized yet."), *GetReadableName(), InUpdateFlags);
+		UE_LOG(LogText3D, Verbose, TEXT("[%s] Rebuild called with update flags (%u) while text is not initialized yet."), *GetReadableName(), InUpdateFlags);
 		InUpdateFlags = EText3DRendererFlags::All;
 		bInitialized = true;
 	}
diff --git a/Engine/Plugins/Experimental/Text3D/Source/Text3D/Public/Extensions/Text3DDefaultGeometryExtension.h b/Engine/Plugins/Experimental/Text3D/Source/Text3D/Public/Extensions/Text3DDefaultGeometryExtension.h
--- a/Engine/Plugins/Experimental/Text3D/Source/Text3D/Public/Extensions/Text3DDefaultGeometryExtension.h
+++ b/Engine/Plugins/Experimental/Text3D/Source/Text3D/Public/Extensions/Text3DDefaultGeometryExtension.h
@@ -102,7 +102,7 @@ protected:
 	//~ End UObject
 
 	//~ Begin UText3DGeometryExtensionBase
-	virtual UStaticMesh* FindOrLoadGlyphMesh(UText3DCharacterBase* InCharacter) const override;
+	virtual const FText3DCachedMesh* FindOrLoadGlyphMesh(UText3DCharacterBase* InCharacter) const override;
 	virtual EText3DHorizontalTextAlignment GetGlyphHAlignment() const override;
 	virtual EText3DVerticalTextAlignment GetGlyphVAlignment() const override;
 	virtual FTypefaceFontData* GetTypefaceFontData() const override;
diff --git a/Engine/Plugins/Experimental/Text3D/Source/Text3D/Public/Extensions/Text3DGeometryExtensionBase.h b/Engine/Plugins/Experimental/Text3D/Source/Text3D/Public/Extensions/Text3DGeometryExtensionBase.h
--- a/Engine/Plugins/Experimental/Text3D/Source/Text3D/Public/Extensions/Text3DGeometryExtensionBase.h
+++ b/Engine/Plugins/Experimental/Text3D/Source/Text3D/Public/Extensions/Text3DGeometryExtensionBase.h
@@ -9,6 +9,7 @@
 class FFreeTypeFace;
 class UStaticMesh;
 class UText3DCharacterBase;
+struct FText3DCachedMesh;
 struct FTypefaceFontData;
 
 /** Extension that handles geometry data for Text3D */
@@ -22,7 +23,7 @@ public:
 		: UText3DExtensionBase(/** Priority */1)
 	{}
 
-	virtual UStaticMesh* FindOrLoadGlyphMesh(UText3DCharacterBase* InCharacter) const
+	virtual const FText3DCachedMesh* FindOrLoadGlyphMesh(UText3DCharacterBase* InCharacter) const
 	{
 		return nullptr;
 	}
diff --git a/Engine/Plugins/Experimental/Text3D/Source/Text3D/Public/Text3DComponent.h b/Engine/Plugins/Experimental/Text3D/Source/Text3D/Public/Text3DComponent.h
--- a/Engine/Plugins/Experimental/Text3D/Source/Text3D/Public/Text3DComponent.h
+++ b/Engine/Plugins/Experimental/Text3D/Source/Text3D/Public/Text3DComponent.h
@@ -398,6 +398,9 @@ protected:
 
 	/** Legacy : used to resolve font by their names */
 	TEXT3D_API void SetFontByName(const FString& InFontName);
+
+	UFUNCTION(CallInEditor, Category="Utilities")
+	void ForceUpdateText();
 #endif
 
 	/** Get the type faces supported for the current font */
diff --git a/Engine/Plugins/Experimental/Text3D/Source/Text3D/Public/Text3DTypes.h b/Engine/Plugins/Experimental/Text3D/Source/Text3D/Public/Text3DTypes.h
--- a/Engine/Plugins/Experimental/Text3D/Source/Text3D/Public/Text3DTypes.h
+++ b/Engine/Plugins/Experimental/Text3D/Source/Text3D/Public/Text3DTypes.h
@@ -20,6 +20,7 @@ enum class EText3DRendererFlags : uint8
 	/** Update everything */
 	All = Geometry | Layout | Visibility | Material
 };
+ENUM_CLASS_FLAGS(EText3DRendererFlags);
 
 UENUM()
 enum class EText3DMaterialStyle : uint8 
@@ -244,6 +245,9 @@ namespace UE::Text3D
 
 			/** Current flag being updated */
 			EText3DRendererFlags CurrentFlag = EText3DRendererFlags::None;
+
+			/** Whether CurrentFlag is the last flag of UpdateFlags */
+			bool bIsLastFlag = false;
 		};
 	}
 }
\ No newline at end of file
diff --git a/Engine/Plugins/Experimental/Text3D/Source/Text3D/Text3D.Build.cs b/Engine/Plugins/Experimental/Text3D/Source/Text3D/Text3D.Build.cs
--- a/Engine/Plugins/Experimental/Text3D/Source/Text3D/Text3D.Build.cs
+++ b/Engine/Plugins/Experimental/Text3D/Source/Text3D/Text3D.Build.cs
@@ -22,6 +22,8 @@ public class Text3D : ModuleRules
 			"DeveloperSettings",
 			"GeometryAlgorithms",
 			"GeometryCore",
+			"GeometryFramework",
+			"GeometryScriptingCore",
             "ICU",
             "MeshDescription",
             "Slate",
diff --git a/Engine/Plugins/Experimental/Text3D/Text3D.uplugin b/Engine/Plugins/Experimental/Text3D/Text3D.uplugin
--- a/Engine/Plugins/Experimental/Text3D/Text3D.uplugin
+++ b/Engine/Plugins/Experimental/Text3D/Text3D.uplugin
@@ -31,6 +31,11 @@
       "Enabled": true,
 	  "TargetDenyList": [ "Server" ]
     },
+	{
+	  "Name": "GeometryScripting", 
+	  "Enabled": true,
+      "TargetDenyList": [ "Server" ]
+	},
     {
       "Name": "GeometryMask",
       "Enabled": true
diff --git a/Engine/Plugins/VirtualProduction/Avalanche/Source/AvalancheModifiers/Private/Modifiers/AvaGeometryBaseModifier.cpp b/Engine/Plugins/VirtualProduction/Avalanche/Source/AvalancheModifiers/Private/Modifiers/AvaGeometryBaseModifier.cpp
--- a/Engine/Plugins/VirtualProduction/Avalanche/Source/AvalancheModifiers/Private/Modifiers/AvaGeometryBaseModifier.cpp
+++ b/Engine/Plugins/VirtualProduction/Avalanche/Source/AvalancheModifiers/Private/Modifiers/AvaGeometryBaseModifier.cpp
@@ -59,6 +59,41 @@ void UAvaGeometryBaseModifier::RestorePreState()
 	}
 }
 
+void UAvaGeometryBaseModifier::OnModifierEnabled(EActorModifierCoreEnableReason InReason)
+{
+	Super::OnModifierEnabled(InReason);
+
+	if (UDynamicMeshComponent* const DynMeshComp = GetMeshComponent())
+	{
+		DynMeshComp->OnMeshChanged.AddUObject(this, &UAvaGeometryBaseModifier::OnMeshChanged);
+	}
+}
+
+void UAvaGeometryBaseModifier::OnModifierDisabled(EActorModifierCoreDisableReason InReason)
+{
+	Super::OnModifierDisabled(InReason);
+
+	if (UDynamicMeshComponent* const DynMeshComp = GetMeshComponent())
+	{
+		DynMeshComp->OnMeshChanged.RemoveAll(this);
+	}
+}
+
+void UAvaGeometryBaseModifier::OnMeshChanged()
+{
+	if (const UActorModifierCoreStack* RootStack = GetRootModifierStack())
+	{
+		if (RootStack->IsModifierIdle()
+			&& IsModifierEnabled()
+			&& RootStack->IsModifierInitialized()
+			&& !RootStack->IsModifierExecutionLocked())
+		{
+			PreModifierCachedMesh.Reset();
+			MarkModifierDirty();
+		}
+	}
+}
+
 UDynamicMeshComponent* UAvaGeometryBaseModifier::GetMeshComponent() const
 {
 	if (!MeshComponent.IsValid())
diff --git a/Engine/Plugins/VirtualProduction/Avalanche/Source/AvalancheModifiers/Public/Modifiers/AvaGeometryBaseModifier.h b/Engine/Plugins/VirtualProduction/Avalanche/Source/AvalancheModifiers/Public/Modifiers/AvaGeometryBaseModifier.h
--- a/Engine/Plugins/VirtualProduction/Avalanche/Source/AvalancheModifiers/Public/Modifiers/AvaGeometryBaseModifier.h
+++ b/Engine/Plugins/VirtualProduction/Avalanche/Source/AvalancheModifiers/Public/Modifiers/AvaGeometryBaseModifier.h
@@ -30,8 +30,12 @@ protected:
 	virtual bool IsModifierReady() const override;
 	virtual void SavePreState() override;
 	virtual void RestorePreState() override;
+	virtual void OnModifierEnabled(EActorModifierCoreEnableReason InReason) override;
+	virtual void OnModifierDisabled(EActorModifierCoreDisableReason InReason) override;
 	//~ End UActorModifierCoreBase
 
+	void OnMeshChanged();
+
 	/** Gets the cached dynamic mesh component. If null, find the first attached and caches it */
 	UDynamicMeshComponent* GetMeshComponent() const;
 
