MotionDesign : Text3D
- Fixed dynamic mesh renderer crash when text gets regenerated due to update flags not triggering a complete refresh of the render proxy
- Added game ready player collisions for static and dynamic mesh

diff --git a/Engine/Plugins/Experimental/Text3D/Source/Text3D/Private/GeometryBuilders/Text3DGlyph.cpp b/Engine/Plugins/Experimental/Text3D/Source/Text3D/Private/GeometryBuilders/Text3DGlyph.cpp
--- a/Engine/Plugins/Experimental/Text3D/Source/Text3D/Private/GeometryBuilders/Text3DGlyph.cpp
+++ b/Engine/Plugins/Experimental/Text3D/Source/Text3D/Private/GeometryBuilders/Text3DGlyph.cpp
@@ -8,6 +8,7 @@
 #include "Materials/Material.h"
 #include "Materials/MaterialInterface.h"
 #include "MeshDescription.h"
+#include "PhysicsEngine/BodySetup.h"
 #include "Text3DInternalTypes.h"
 #include "UDynamicMesh.h"
 #include "Subsystems/Text3DEngineSubsystem.h"
@@ -69,6 +70,12 @@ void FText3DGlyph::Build(FText3DCachedMesh& InMesh, UMaterialInterface* InDefaul
 
 		InMesh.StaticMesh->bAllowCPUAccess = true;
 		InMesh.StaticMesh->BuildFromMeshDescriptions(MeshDescriptions, BuildParams);
+
+		// Setup collision
+		if (UBodySetup* Body = InMesh.StaticMesh->GetBodySetup())
+		{
+			Body->CollisionTraceFlag = ECollisionTraceFlag::CTF_UseComplexAsSimple;
+		}
 	}
 
 	{
diff --git a/Engine/Plugins/Experimental/Text3D/Source/Text3D/Private/Renderers/DynamicMesh/Text3DDynamicMeshRenderer.cpp b/Engine/Plugins/Experimental/Text3D/Source/Text3D/Private/Renderers/DynamicMesh/Text3DDynamicMeshRenderer.cpp
--- a/Engine/Plugins/Experimental/Text3D/Source/Text3D/Private/Renderers/DynamicMesh/Text3DDynamicMeshRenderer.cpp
+++ b/Engine/Plugins/Experimental/Text3D/Source/Text3D/Private/Renderers/DynamicMesh/Text3DDynamicMeshRenderer.cpp
@@ -6,6 +6,7 @@
 #include "Components/DynamicMeshComponent.h"
 #include "DynamicMesh/MeshTransforms.h"
 #include "DynamicMeshEditor.h"
+#include "Engine/CollisionProfile.h"
 #include "Extensions/Text3DGeometryExtensionBase.h"
 #include "Extensions/Text3DLayoutExtensionBase.h"
 #include "Extensions/Text3DMaterialExtensionBase.h"
@@ -16,20 +17,28 @@
 
 UText3DDynamicMeshRenderer::FScopedMeshEditor::~FScopedMeshEditor()
 {
-	if (bUpdated && IsValid(MeshComponent))
+	if (UpdateType != EMeshEditorUpdateType::None && IsValid(MeshComponent))
 	{
-		MeshComponent->NotifyMeshVertexAttributesModified();
+		if (UpdateType == EMeshEditorUpdateType::Fast)
+		{
+			MeshComponent->NotifyMeshVertexAttributesModified();
+		}
+		else if (UpdateType == EMeshEditorUpdateType::Full)
+		{
+			MeshComponent->NotifyMeshUpdated();
+		}
+
 		MeshComponent->OnMeshChanged.Broadcast();
 	}
 }
 
-void UText3DDynamicMeshRenderer::FScopedMeshEditor::EditMesh(const TFunctionRef<void(UE::Geometry::FDynamicMesh3&)>& InFunctor)
+void UText3DDynamicMeshRenderer::FScopedMeshEditor::EditMesh(const TFunctionRef<void(UE::Geometry::FDynamicMesh3&)>& InFunctor, EMeshEditorUpdateType InType)
 {
 	if (MeshComponent && MeshComponent->GetDynamicMesh())
 	{
 		constexpr bool bDeferChanges = true;
 		MeshComponent->GetDynamicMesh()->EditMesh(InFunctor, EDynamicMeshChangeType::GeneralEdit, EDynamicMeshAttributeChangeFlags::Unknown, bDeferChanges);
-		bUpdated = true;
+		UpdateType = InType > UpdateType ? InType : UpdateType;
 	}
 }
 
@@ -46,8 +55,14 @@ void UText3DDynamicMeshRenderer::OnCreate()
 	if (!DynamicMeshComponent)
 	{
 		UText3DComponent* TextComponent = GetText3DComponent();
-		DynamicMeshComponent = NewObject<UDynamicMeshComponent>(this, TEXT("Text3DDynamicMesh"), RF_Transient);
+		DynamicMeshComponent = NewObject<UDynamicMeshComponent>(this, TEXT("Text3DDynamicMesh"));
 		DynamicMeshComponent->SetupAttachment(TextComponent);
+
+		// Setup collision
+		DynamicMeshComponent->SetComplexAsSimpleCollisionEnabled(true);
+		DynamicMeshComponent->SetDeferredCollisionUpdatesEnabled(true);
+		DynamicMeshComponent->SetCollisionProfileName(UCollisionProfile::BlockAllDynamic_ProfileName);
+
 		ClearMesh();
 	}
 }
@@ -60,7 +75,7 @@ void UText3DDynamicMeshRenderer::OnUpdate(const UE::Text3D::Renderer::FUpdatePar
 	{
 		ScopedMeshEditor = FScopedMeshEditor(DynamicMeshComponent);
 
-		// Restore checkpoint
+		// Restore checkpoint in case mesh was altered externally
 		RestoreMesh(ScopedMeshEditor.GetValue());
 	}
 
@@ -135,7 +150,7 @@ void UText3DDynamicMeshRenderer::OnUpdate(const UE::Text3D::Renderer::FUpdatePar
 
 	if (InParameters.bIsLastFlag)
 	{
-		// Save checkpoint
+		// Save checkpoint to restore mesh on next update
 		SaveMesh(MeshEditor);
 
 		ScopedMeshEditor.Reset();
@@ -288,7 +303,7 @@ void UText3DDynamicMeshRenderer::AppendGlyphMesh(FScopedMeshEditor& InMeshEditor
 
 				GlyphMeshData[InIndex].MaterialSlots = MoveTemp(MaterialSlots);
 			});
-		});
+		}, EMeshEditorUpdateType::Full);
 }
 
 void UText3DDynamicMeshRenderer::ClearMesh()
@@ -325,7 +340,7 @@ void UText3DDynamicMeshRenderer::ClearGlyphMesh(FScopedMeshEditor& InMeshEditor,
 			FDynamicMeshEditor Editor(&InEditMesh);
 			Editor.RemoveTriangles(TrianglesToDelete, /** RemoveIsolatedVertices*/ true);
 			InEditMesh.CompactInPlace();
-		});
+		}, EMeshEditorUpdateType::Full);
 }
 
 void UText3DDynamicMeshRenderer::SetGlyphMeshTransform(FScopedMeshEditor& InMeshEditor, uint16 InIndex, const FTransform& InTransform, bool bInVisible)
@@ -410,7 +425,7 @@ void UText3DDynamicMeshRenderer::SetGlyphMeshTransform(FScopedMeshEditor& InMesh
 				Tangents->SetElement(TId, (FVector3f)NewTangent.GetSafeNormal());	
 			}
 		}
-	});
+	}, EMeshEditorUpdateType::Fast);
 
 	GlyphData.CurrentTransform = TargetTransform;
 }
@@ -431,6 +446,6 @@ void UText3DDynamicMeshRenderer::RestoreMesh(FScopedMeshEditor& InMeshEditor)
 		{
 			EditMesh = MoveTemp(CachedMesh.GetValue());
 			CachedMesh.Reset();
-		});
+		}, EMeshEditorUpdateType::Fast);
 	}
 }
diff --git a/Engine/Plugins/Experimental/Text3D/Source/Text3D/Private/Renderers/DynamicMesh/Text3DDynamicMeshRenderer.h b/Engine/Plugins/Experimental/Text3D/Source/Text3D/Private/Renderers/DynamicMesh/Text3DDynamicMeshRenderer.h
--- a/Engine/Plugins/Experimental/Text3D/Source/Text3D/Private/Renderers/DynamicMesh/Text3DDynamicMeshRenderer.h
+++ b/Engine/Plugins/Experimental/Text3D/Source/Text3D/Private/Renderers/DynamicMesh/Text3DDynamicMeshRenderer.h
@@ -31,6 +31,14 @@ protected:
 		FBox Bounds = FBox(ForceInitToZero);
 	};
 
+	/** Mesh update type */
+	enum class EMeshEditorUpdateType : uint8
+	{
+		None,
+		Fast,
+		Full
+	};
+
 	/** Used to push render changes when destroyed */
 	struct FScopedMeshEditor
 	{
@@ -39,12 +47,12 @@ protected:
 		{}
 
 		~FScopedMeshEditor();
-		void EditMesh(const TFunctionRef<void(UE::Geometry::FDynamicMesh3&)>& InFunctor);
+		void EditMesh(const TFunctionRef<void(UE::Geometry::FDynamicMesh3&)>& InFunctor, EMeshEditorUpdateType InType);
 		void ProcessMesh(const TFunctionRef<void(const UE::Geometry::FDynamicMesh3&)>& InFunctor) const;
 
 	private:
 		UDynamicMeshComponent* MeshComponent = nullptr;
-		bool bUpdated = false;
+		EMeshEditorUpdateType UpdateType = EMeshEditorUpdateType::None;
 	};
 
 	//~ Begin UText3DRendererBase
diff --git a/Engine/Plugins/Experimental/Text3D/Source/Text3DEditor/Private/Widgets/SText3DEditorFontField.cpp b/Engine/Plugins/Experimental/Text3D/Source/Text3DEditor/Private/Widgets/SText3DEditorFontField.cpp
--- a/Engine/Plugins/Experimental/Text3D/Source/Text3DEditor/Private/Widgets/SText3DEditorFontField.cpp
+++ b/Engine/Plugins/Experimental/Text3D/Source/Text3DEditor/Private/Widgets/SText3DEditorFontField.cpp
@@ -104,14 +104,14 @@ void SText3DEditorFontField::Construct(const FArguments& InArgs)
 
 void SText3DEditorFontField::UpdateFont()
 {
-	TSharedPtr<FString> FontItem = FontItemWeak.Pin();
+	const TSharedPtr<FString> FontItem = FontItemWeak.Pin();
 
 	if (!FontItem.IsValid())
 	{
 		return;
 	}
 
-	UText3DEditorFontSubsystem* FontEditorSubsystem = UText3DEditorFontSubsystem::Get();
+	const UText3DEditorFontSubsystem* FontEditorSubsystem = UText3DEditorFontSubsystem::Get();
 	const FText3DEditorFont* EditorFont = FontEditorSubsystem->GetEditorFont(*FontItem.Get());
 	
 	if (EditorFont && EditorFont->Font)
@@ -132,7 +132,7 @@ void SText3DEditorFontField::UpdateFont()
 
 ECheckBoxState SText3DEditorFontField::GetFavoriteState() const
 {
-	if (TSharedPtr<FString> FontItem = FontItemWeak.Pin())
+	if (const TSharedPtr<FString> FontItem = FontItemWeak.Pin())
 	{
 		if (const UText3DProjectSettings* Text3DSettings = UText3DProjectSettings::Get())
 		{
@@ -150,7 +150,7 @@ FSlateColor SText3DEditorFontField::GetToggleFavoriteColor() const
 
 FReply SText3DEditorFontField::OnToggleFavoriteClicked()
 {
-	TSharedPtr<FString> FontItem = FontItemWeak.Pin();
+	const TSharedPtr<FString> FontItem = FontItemWeak.Pin();
 
 	if (!FontItem.IsValid())
 	{
@@ -160,7 +160,7 @@ FReply SText3DEditorFontField::OnToggleFavoriteClicked()
 	UText3DProjectSettings* Text3DSettings = UText3DProjectSettings::GetMutable();
 	if (!Text3DSettings)
 	{
-		FReply::Handled();
+		return FReply::Handled();
 	}
 
 	if (GetFavoriteState() == ECheckBoxState::Checked)
@@ -185,7 +185,7 @@ EVisibility SText3DEditorFontField::GetFavoriteVisibility() const
 
 FReply SText3DEditorFontField::OnBrowseToAssetClicked() const
 {
-	TSharedPtr<FString> FontItem = FontItemWeak.Pin();
+	const TSharedPtr<FString> FontItem = FontItemWeak.Pin();
 
 	if (!FontItem.IsValid())
 	{
@@ -206,7 +206,7 @@ FReply SText3DEditorFontField::OnBrowseToAssetClicked() const
 
 EVisibility SText3DEditorFontField::GetLocallyAvailableIconVisibility() const
 {
-	TSharedPtr<FString> FontItem = FontItemWeak.Pin();
+	const TSharedPtr<FString> FontItem = FontItemWeak.Pin();
 
 	if (!FontItem.IsValid())
 	{
diff --git a/Engine/Plugins/VirtualProduction/Avalanche/Source/AvalancheShapes/Private/AvaShapeActor.cpp b/Engine/Plugins/VirtualProduction/Avalanche/Source/AvalancheShapes/Private/AvaShapeActor.cpp
--- a/Engine/Plugins/VirtualProduction/Avalanche/Source/AvalancheShapes/Private/AvaShapeActor.cpp
+++ b/Engine/Plugins/VirtualProduction/Avalanche/Source/AvalancheShapes/Private/AvaShapeActor.cpp
@@ -17,9 +17,12 @@ AAvaShapeActor::AAvaShapeActor()
 	PrimaryActorTick.bTickEvenWhenPaused   = false;
 
 	ShapeMeshComponent = CreateDefaultSubobject<UDynamicMeshComponent>(ShapeComponentName);
-	ShapeMeshComponent->SetCollisionProfileName(UCollisionProfile::BlockAll_ProfileName);
+
+	// Setup collision
+	ShapeMeshComponent->SetCollisionProfileName(UCollisionProfile::BlockAllDynamic_ProfileName);
 	ShapeMeshComponent->SetGenerateOverlapEvents(true);
 	ShapeMeshComponent->SetCanEverAffectNavigation(false);
+	ShapeMeshComponent->SetComplexAsSimpleCollisionEnabled(true);
 
 	SetRootComponent(ShapeMeshComponent);
 
