MotionDesign : Text3D
- Added text outline fill or stroke type, uses bevel material slot since bevel and outline are incompatible with each other

diff --git a/Engine/Plugins/Experimental/Text3D/Source/Text3D/Private/Extensions/Text3DDefaultGeometryExtension.cpp b/Engine/Plugins/Experimental/Text3D/Source/Text3D/Private/Extensions/Text3DDefaultGeometryExtension.cpp
--- a/Engine/Plugins/Experimental/Text3D/Source/Text3D/Private/Extensions/Text3DDefaultGeometryExtension.cpp
+++ b/Engine/Plugins/Experimental/Text3D/Source/Text3D/Private/Extensions/Text3DDefaultGeometryExtension.cpp
@@ -32,6 +32,17 @@ void UText3DDefaultGeometryExtension::SetOutline(const float Value)
 	OnGeometryOptionsChanged();
 }
 
+void UText3DDefaultGeometryExtension::SetOutlineType(EText3DOutlineType InType)
+{
+	if (OutlineType == InType)
+	{
+		return;
+	}
+
+	OutlineType = InType;
+	OnGeometryOptionsChanged();
+}
+
 void UText3DDefaultGeometryExtension::SetPivotHAlignment(EText3DHorizontalTextAlignment InPivot)
 {
 	if (PivotHAlignment == InPivot)
@@ -126,6 +137,7 @@ void UText3DDefaultGeometryExtension::PostEditChangeProperty(FPropertyChangedEve
 		GET_MEMBER_NAME_CHECKED(UText3DDefaultGeometryExtension, BevelSegments),
 		GET_MEMBER_NAME_CHECKED(UText3DDefaultGeometryExtension, bUseOutline),
 		GET_MEMBER_NAME_CHECKED(UText3DDefaultGeometryExtension, Outline),
+		GET_MEMBER_NAME_CHECKED(UText3DDefaultGeometryExtension, OutlineType),
 		GET_MEMBER_NAME_CHECKED(UText3DDefaultGeometryExtension, PivotHAlignment),
 		GET_MEMBER_NAME_CHECKED(UText3DDefaultGeometryExtension, PivotVAlignment),
 	};
@@ -203,6 +215,7 @@ EText3DExtensionResult UText3DDefaultGeometryExtension::PreRendererUpdate(const
 		.BevelSegments = GetBevelSegments(),
 		.bOutline = GetUseOutline(),
 		.OutlineExpand = GetOutline(),
+		.OutlineType = GetOutlineType(),
 		.TypefaceIndex = TypefaceIndex,
 		.PivotOffset = GetPivotOffset(),
 	};
@@ -220,11 +233,9 @@ EText3DExtensionResult UText3DDefaultGeometryExtension::PostRendererUpdate(const
 void UText3DDefaultGeometryExtension::OnGeometryOptionsChanged()
 {
 	// Extrude
-
 	Extrude = FMath::Max(0.0f, Extrude);
 
 	// Bevels
-	
 	Bevel = FMath::Clamp(Bevel, 0.f, GetMaxBevel());
 
 	int32 MaxBevelSegments = 1;
diff --git a/Engine/Plugins/Experimental/Text3D/Source/Text3D/Private/GeometryBuilders/Text3DGlyphMeshBuilder.cpp b/Engine/Plugins/Experimental/Text3D/Source/Text3D/Private/GeometryBuilders/Text3DGlyphMeshBuilder.cpp
--- a/Engine/Plugins/Experimental/Text3D/Source/Text3D/Private/GeometryBuilders/Text3DGlyphMeshBuilder.cpp
+++ b/Engine/Plugins/Experimental/Text3D/Source/Text3D/Private/GeometryBuilders/Text3DGlyphMeshBuilder.cpp
@@ -17,18 +17,24 @@ FText3DGlyphMeshBuilder::FText3DGlyphMeshBuilder() :
 {
 }
 
-void FText3DGlyphMeshBuilder::CreateMeshes(const TText3DGlyphContourNodeShared& Root, const float Extrude, const float Bevel, const EText3DBevelType Type, const int32 BevelSegments, const bool bOutline, const float OutlineExpand)
+void FText3DGlyphMeshBuilder::CreateMeshes(const TText3DGlyphContourNodeShared& Root, const float InExtrude, const float InBevel, const EText3DBevelType InBevelType, const int32 InBevelSegments, const bool bInOutline, const float InOutlineExpand, const EText3DOutlineType InOutlineType)
 {
-	CreateFrontMesh(Root, bOutline, OutlineExpand);
+	CreateFrontMesh(Root, bInOutline, InOutlineExpand, InOutlineType);
+
 	if (Contours->Num() == 0)
 	{
 		return;
 	}
 
 	const bool bFlipNormals = FMath::Sign(Data->GetPlannedExpand()) < 0;
-	const float BevelLocal = bOutline ? 0.f : Bevel;
-	CreateBevelMesh(BevelLocal, Type, BevelSegments);
-	CreateExtrudeMesh(Extrude, BevelLocal, Type, bFlipNormals);
+	const float BevelLocal = bInOutline ? 0.f : InBevel;
+
+	if (!bInOutline)
+	{
+		CreateBevelMesh(BevelLocal, InBevelType, InBevelSegments);	
+	}
+
+	CreateExtrudeMesh(InExtrude, BevelLocal, InBevelType, bFlipNormals);
 }
 
 void FText3DGlyphMeshBuilder::SetFrontAndBevelTextureCoordinates(const float Bevel)
@@ -146,7 +152,7 @@ FVector FText3DGlyphMeshBuilder::GetMeshOffset() const
 	return MeshOffset;
 }
 
-void FText3DGlyphMeshBuilder::CreateFrontMesh(const TText3DGlyphContourNodeShared& Root, const bool bOutline, const float& OutlineExpand)
+void FText3DGlyphMeshBuilder::CreateFrontMesh(const TText3DGlyphContourNodeShared& Root, const bool bOutline, const float OutlineExpand, const EText3DOutlineType InOutlineType)
 {
 	int32 VertexCount = 0;
 	AddToVertexCount(Root, VertexCount);
@@ -156,13 +162,13 @@ void FText3DGlyphMeshBuilder::CreateFrontMesh(const TText3DGlyphContourNodeShare
 	Contours = MakeShared<FText3DGlyphContourList>();
 
 	int32 VertexIndex = Data->AddVertices(VertexCount);
-	TriangulateAndConvert(Root, VertexIndex, bOutline);
+	TriangulateAndConvert(Root, VertexIndex, bOutline, InOutlineType);
 
 	Contours->Initialize(Data);
 
 	if (bOutline)
 	{
-		MakeOutline(OutlineExpand);
+		CreateOutlineMesh(OutlineExpand);
 	}
 }
 
@@ -411,7 +417,7 @@ void FText3DGlyphMeshBuilder::AddToVertexCount(const TText3DGlyphContourNodeShar
 	}
 }
 
-void FText3DGlyphMeshBuilder::TriangulateAndConvert(const TText3DGlyphContourNodeShared& Node, int32& OutVertexIndex, const bool bOutline)
+void FText3DGlyphMeshBuilder::TriangulateAndConvert(const TText3DGlyphContourNodeShared& Node, int32& OutVertexIndex, const bool bOutline, const EText3DOutlineType InOutlineType)
 {
 	// If this is solid region
 	if (!Node->bClockwise)
@@ -422,7 +428,7 @@ void FText3DGlyphMeshBuilder::TriangulateAndConvert(const TText3DGlyphContourNod
 
 		const TSharedPtr<FText3DGlyphContourList> ContoursLocal = Contours;
 		const TSharedRef<FText3DGlyphData> DataLocal = Data;
-		auto ProcessContour = [ContoursLocal, DataLocal, &VertexCount, &Triangulation, bOutline](const TText3DGlyphContourNodeShared NodeIn)
+		auto ProcessContour = [ContoursLocal, DataLocal, &VertexCount, &Triangulation, bOutline, InOutlineType](const TText3DGlyphContourNodeShared NodeIn)
 		{
 			// Create contour in old format
 			FText3DGlyphContour& Contour = ContoursLocal->Add();
@@ -445,13 +451,12 @@ void FText3DGlyphMeshBuilder::TriangulateAndConvert(const TText3DGlyphContourNod
 			VertexCount += Polygon.VertexCount();
 
 			// Add contour to triangulation
-			if (!bOutline)
+			if (!bOutline || InOutlineType != EText3DOutlineType::Stroke)
 			{
 				Triangulation.Add(Polygon, NodeIn->bClockwise);
 			}
 		};
 
-
 		// Outer
 		ProcessContour(Node);
 
@@ -461,7 +466,7 @@ void FText3DGlyphMeshBuilder::TriangulateAndConvert(const TText3DGlyphContourNod
 			ProcessContour(Child);
 		}
 
-		if (!bOutline)
+		if (!bOutline || InOutlineType != EText3DOutlineType::Stroke)
 		{
 			Triangulation.Triangulate();
 			const TArray<FIndex3i>& Triangles = Triangulation.Triangles;
@@ -479,12 +484,14 @@ void FText3DGlyphMeshBuilder::TriangulateAndConvert(const TText3DGlyphContourNod
 	// Continue with children
 	for (const TText3DGlyphContourNodeShared& Child : Node->Children)
 	{
-		TriangulateAndConvert(Child, OutVertexIndex, bOutline);
+		TriangulateAndConvert(Child, OutVertexIndex, bOutline, InOutlineType);
 	}
 }
 
-void FText3DGlyphMeshBuilder::MakeOutline(float OutlineExpand)
+void FText3DGlyphMeshBuilder::CreateOutlineMesh(float InOutlineExpand)
 {
+	Data->SetCurrentGroup(EText3DGroupType::Bevel, InOutlineExpand);
+
 	FText3DGlyphContourList InitialContours = *Contours;
 
 	for (FText3DGlyphContour& Contour : InitialContours)
@@ -497,7 +504,6 @@ void FText3DGlyphMeshBuilder::MakeOutline(float OutlineExpand)
 			Point->Normal *= -1.f;
 		}
 
-
 		const FText3DGlyphPartPtr First = Contour[0];
 		const FText3DGlyphPartPtr Last = Contour.Last();
 
@@ -511,14 +517,11 @@ void FText3DGlyphMeshBuilder::MakeOutline(float OutlineExpand)
 		Last->TangentX = -FirstTangentX;
 	}
 
-
 	const FVector2D Normal = {0.f, -1.f};
-	BevelLinear(0.f, OutlineExpand, Normal, Normal, false);
-
+	BevelLinear(0.f, InOutlineExpand, Normal, Normal, false);
 
 	Contours->Reset();
 
-
 	TDoubleLinkedList<FText3DGlyphContour>::TDoubleLinkedListNode* Node = InitialContours.GetHead();
 	while (Node)
 	{
diff --git a/Engine/Plugins/Experimental/Text3D/Source/Text3D/Private/GeometryBuilders/Text3DGlyphMeshBuilder.h b/Engine/Plugins/Experimental/Text3D/Source/Text3D/Private/GeometryBuilders/Text3DGlyphMeshBuilder.h
--- a/Engine/Plugins/Experimental/Text3D/Source/Text3D/Private/GeometryBuilders/Text3DGlyphMeshBuilder.h
+++ b/Engine/Plugins/Experimental/Text3D/Source/Text3D/Private/GeometryBuilders/Text3DGlyphMeshBuilder.h
@@ -23,14 +23,15 @@ public:
 	/**
 	 * CreateMeshes
 	 * @param Root - Tree of contours.	 
-	 * @param Extrude - Orthogonal (to front cap) offset value.
-	 * @param Bevel - Bevel value (bevel happens before extrude).
-	 * @param Type - Defines shape of beveled part.
-	 * @param BevelSegments - Segments count.
-	 * @param bOutline - Front face has outline and is not filled.
-	 * @param OutlineExpand - Offsets the outline by the specified amount.
+	 * @param InExtrude - Orthogonal (to front cap) offset value.
+	 * @param InBevel - Bevel value (bevel happens before extrude).
+	 * @param InBevelType - Defines shape of beveled part.
+	 * @param InBevelSegments - Segments count.
+	 * @param bInOutline - Front face has outline and is not filled.
+	 * @param InOutlineExpand - Offsets the outline by the specified amount.
+	 * @param InOutlineType - Defines the outline behavior.
 	 */
-	void CreateMeshes(const TText3DGlyphContourNodeShared& Root, const float Extrude, const float Bevel, const EText3DBevelType Type, const int32 BevelSegments, const bool bOutline, const float OutlineExpand = 0.5f);
+	void CreateMeshes(const TText3DGlyphContourNodeShared& Root, const float InExtrude, const float InBevel, const EText3DBevelType InBevelType, const int32 InBevelSegments, const bool bInOutline, const float InOutlineExpand, const EText3DOutlineType InOutlineType);
 	void SetFrontAndBevelTextureCoordinates(const float Bevel);
 	void MirrorGroups(const float Extrude);
 
@@ -57,8 +58,10 @@ private:
 	 * @param Root - Tree root.
 	 * @param bOutline - Front face has outline and is not filled.
 	 * @param OutlineExpand - Offsets the outline by the specified amount.
+	 * @param InOutlineType - Defines the outline behavior.
 	 */
-	void CreateFrontMesh(const TText3DGlyphContourNodeShared& Root, const bool bOutline, const float& OutlineExpand = 0.5f);
+	void CreateFrontMesh(const TText3DGlyphContourNodeShared& Root, const bool bOutline, const float OutlineExpand, const EText3DOutlineType InOutlineType);
+
 	/**
 	 * Create 'Bevel' part of glyph (actually half of it, will be mirrored later).
 	 * @param Bevel - Bevel value (bevel happens before extrude).
@@ -66,6 +69,7 @@ private:
 	 * @param BevelSegments - Segments count.
 	 */
 	void CreateBevelMesh(const float Bevel, const EText3DBevelType Type, const int32 BevelSegments);
+
 	/**
 	 * Create 'Extrude' part of glyph.
 	 * @param Extrude - Extrude value.
@@ -75,8 +79,16 @@ private:
 	 */
 	void CreateExtrudeMesh(float Extrude, float Bevel, const EText3DBevelType Type, bool bFlipNormals = false);
 
-	void MirrorGroup(const EText3DGroupType TypeIn, const EText3DGroupType TypeOut, const float Extrude);
+	/**
+	 * Create 'Outline' part of glyph. Not compatible with Bevel
+	 * @param InOutlineExpand - outline expansion value.
+	 */
+	void CreateOutlineMesh(float InOutlineExpand);
 
+	/**
+	 * Copies and reflect source group into a destination group based on extrusion value
+	 */
+	void MirrorGroup(const EText3DGroupType TypeIn, const EText3DGroupType TypeOut, const float Extrude);
 
 	/**
 	 * Recursively compute vertex count in tree.
@@ -84,15 +96,14 @@ private:
 	 * @param OutVertexCount - Pointer to counter.
 	 */
 	void AddToVertexCount(const TText3DGlyphContourNodeShared& Node, int32& OutVertexCount);
+
 	/**
 	 * Triangulate solid region with holes and convert contours to old format.
 	 * @param Node - Tree node.
 	 * @param OutVertexIndex - Pointer to last used vertex index.
 	 * @param bOutline - Front face has outline and is not filled.
 	 */
-	void TriangulateAndConvert(const TText3DGlyphContourNodeShared& Node, int32& OutVertexIndex, const bool bOutline);
-
-	void MakeOutline(float OutlineExpand = 0.5f);
+	void TriangulateAndConvert(const TText3DGlyphContourNodeShared& Node, int32& OutVertexIndex, const bool bOutline, const EText3DOutlineType InOutlineType);
 
 	void BevelLinearWithSegments(const float Extrude, const float Expand, const int32 BevelSegments, const FVector2D Normal);
 	void BevelCurve(const float Angle, const int32 BevelSegments, TFunction<FVector2D(const float CurrentCos, const float CurrentSin, const float NextCos, const float Next)> ComputeOffset);
@@ -130,6 +141,7 @@ private:
 	 * @param TextureCoordinates - Texcture coordinates of added vertices.
 	 */
 	void ExpandPoint(const FText3DGlyphPartPtr& Point, const FVector2D TextureCoordinates = FVector2D(0.f, 0.f));
+
 	/**
 	 * Common code for expanding, vertices are added uninitialized.
 	 * @param Point - Expanded point.
@@ -142,6 +154,7 @@ private:
 	 * @param TextureCoordinates - Texture coordinates of added vertex.
 	 */
 	void AddVertexSmooth(const FText3DGlyphPartConstPtr& Point, const FVector2D TextureCoordinates);
+
 	/**
 	 * Add vertex for sharp point.
 	 * @param Point - Expanded point.
diff --git a/Engine/Plugins/Experimental/Text3D/Source/Text3D/Private/Subsystems/Text3DEngineSubsystem.cpp b/Engine/Plugins/Experimental/Text3D/Source/Text3D/Private/Subsystems/Text3DEngineSubsystem.cpp
--- a/Engine/Plugins/Experimental/Text3D/Source/Text3D/Private/Subsystems/Text3DEngineSubsystem.cpp
+++ b/Engine/Plugins/Experimental/Text3D/Source/Text3D/Private/Subsystems/Text3DEngineSubsystem.cpp
@@ -340,7 +340,16 @@ FText3DCachedMesh* FCachedFontData::GetGlyphMesh(uint32 InGlyphIndex, const FGly
 	CachedMesh.DynamicMesh = NewObject<UDynamicMesh>(TextSubsystem, DynamicMeshName, RF_Transient | RF_Public);
 
 	FText3DGlyphMeshBuilder MeshCreator;
-	MeshCreator.CreateMeshes(Root, InParameters.Extrude, InParameters.Bevel, InParameters.BevelType, InParameters.BevelSegments, InParameters.bOutline, InParameters.OutlineExpand);
+	MeshCreator.CreateMeshes(
+		Root,
+		InParameters.Extrude,
+		InParameters.Bevel,
+		InParameters.BevelType,
+		InParameters.BevelSegments,
+		InParameters.bOutline,
+		InParameters.OutlineExpand,
+		InParameters.OutlineType
+	);
 	MeshCreator.SetFrontAndBevelTextureCoordinates(InParameters.Bevel);
 	MeshCreator.MirrorGroups(InParameters.Extrude);
 	MeshCreator.MovePivot(InParameters.PivotOffset);
diff --git a/Engine/Plugins/Experimental/Text3D/Source/Text3D/Private/Subsystems/Text3DEngineSubsystem.h b/Engine/Plugins/Experimental/Text3D/Source/Text3D/Private/Subsystems/Text3DEngineSubsystem.h
--- a/Engine/Plugins/Experimental/Text3D/Source/Text3D/Private/Subsystems/Text3DEngineSubsystem.h
+++ b/Engine/Plugins/Experimental/Text3D/Source/Text3D/Private/Subsystems/Text3DEngineSubsystem.h
@@ -41,6 +41,9 @@ struct FGlyphMeshParameters
 	UPROPERTY()
 	float OutlineExpand = 0.5f;
 
+	UPROPERTY()
+	EText3DOutlineType OutlineType = EText3DOutlineType::Stroke;
+
 	UPROPERTY()
 	uint32 TypefaceIndex = 0;
 
@@ -56,6 +59,7 @@ struct FGlyphMeshParameters
 		HashParameters = HashCombine(HashParameters, GetTypeHash(InElement.BevelSegments));
 		HashParameters = HashCombine(HashParameters, GetTypeHash(InElement.bOutline));
 		HashParameters = HashCombine(HashParameters, GetTypeHash(InElement.OutlineExpand));
+		HashParameters = HashCombine(HashParameters, GetTypeHash(InElement.OutlineType));
 		HashParameters = HashCombine(HashParameters, GetTypeHash(InElement.TypefaceIndex));
 		HashParameters = HashCombine(HashParameters, GetTypeHash(InElement.PivotOffset));
 		return HashParameters;
diff --git a/Engine/Plugins/Experimental/Text3D/Source/Text3D/Public/Extensions/Text3DDefaultGeometryExtension.h b/Engine/Plugins/Experimental/Text3D/Source/Text3D/Public/Extensions/Text3DDefaultGeometryExtension.h
--- a/Engine/Plugins/Experimental/Text3D/Source/Text3D/Public/Extensions/Text3DDefaultGeometryExtension.h
+++ b/Engine/Plugins/Experimental/Text3D/Source/Text3D/Public/Extensions/Text3DDefaultGeometryExtension.h
@@ -78,6 +78,15 @@ public:
 	UFUNCTION(BlueprintCallable, Category = "Text3D|Geometry")
 	void SetOutline(const float Value);
 
+	UFUNCTION(BlueprintCallable, Category = "Text3D|Geometry")
+	void SetOutlineType(EText3DOutlineType InType);
+
+	UFUNCTION(BlueprintPure, Category = "Text3D|Geometry")
+	EText3DOutlineType GetOutlineType() const
+	{
+		return OutlineType;
+	}
+
 	UFUNCTION(BlueprintCallable, Category = "Text3D|Geometry")
 	void SetPivotHAlignment(EText3DHorizontalTextAlignment InPivot);
 
@@ -138,10 +147,14 @@ protected:
 	UPROPERTY(EditAnywhere, Getter = "GetUseOutline", Setter = "SetUseOutline", Category = "Geometry", meta = (AllowPrivateAccess = "true"))
 	bool bUseOutline = false;
 
-	/** Outline expand/offset amount */
+	/** Outline expand/offset amount, incompatible with bevel */
 	UPROPERTY(EditAnywhere, Getter, Setter, Category = "Geometry", meta = (EditCondition = "bUseOutline", AllowPrivateAccess = "true"))
 	float Outline = 0.5f;
 
+	/** Outline Type */
+	UPROPERTY(EditAnywhere, Getter, Setter, Category = "Geometry", meta = (EditCondition = "bUseOutline", AllowPrivateAccess = "true"))
+	EText3DOutlineType OutlineType = EText3DOutlineType::Stroke;
+
 	/** Pivot alignment of each character */
 	UPROPERTY(EditAnywhere, DisplayName="Character Pivot Alignment Horizontal", Getter, Setter, Category = "Geometry", meta = (AllowPrivateAccess = "true"))
 	EText3DHorizontalTextAlignment PivotHAlignment = EText3DHorizontalTextAlignment::Left;
diff --git a/Engine/Plugins/Experimental/Text3D/Source/Text3D/Public/Text3DTypes.h b/Engine/Plugins/Experimental/Text3D/Source/Text3D/Public/Text3DTypes.h
--- a/Engine/Plugins/Experimental/Text3D/Source/Text3D/Public/Text3DTypes.h
+++ b/Engine/Plugins/Experimental/Text3D/Source/Text3D/Public/Text3DTypes.h
@@ -61,6 +61,13 @@ enum class EText3DBevelType : uint8
 	Engraved
 };
 
+UENUM()
+enum class EText3DOutlineType : uint8
+{
+	Stroke,
+	Fill
+};
+
 enum class EText3DGroupType : uint8
 {
 	Front = 0,
