MotionDesign : Text3D
- Hiding Text3DCharacterTransform from "add component" list to avoid being added to the new Text3D, the UText3DLayoutTransformEffect is the new approach to handle per character transform, keep component to migrate properties to newer system

diff --git a/Engine/Plugins/Experimental/Text3D/Source/Text3D/Private/Characters/Text3DCharacterTransform.cpp b/Engine/Plugins/Experimental/Text3D/Source/Text3D/Private/Characters/Text3DCharacterTransform.cpp
--- a/Engine/Plugins/Experimental/Text3D/Source/Text3D/Private/Characters/Text3DCharacterTransform.cpp
+++ b/Engine/Plugins/Experimental/Text3D/Source/Text3D/Private/Characters/Text3DCharacterTransform.cpp
@@ -1,8 +1,9 @@
 // Copyright Epic Games, Inc. All Rights Reserved.
 
 #include "Characters/Text3DCharacterTransform.h"
+#include "Extensions/Text3DLayoutTransformEffect.h"
 #include "GameFramework/Actor.h"
-#include "Components/StaticMeshComponent.h"
+#include "Logs/Text3DLogs.h"
 #include "Renderers/StaticMeshes/Text3DStaticMeshesRenderer.h"
 #include "Text3DComponent.h"
 
@@ -27,227 +28,6 @@ UText3DCharacterTransform::UText3DCharacterTransform()
 	RotateRange = 50.0f;
 	RotateBegin = FRotator(-90.0f, 0.0f, 0.0f);
 	RotateEnd = FRotator(0.0f, 0.0f, 0.0f);
-
-	bInitialized = false;
-}
-
-void UText3DCharacterTransform::OnRegister()
-{
-	Super::OnRegister();
-
-	if (!bInitialized)
-	{
-		bInitialized = true;
-		ProcessEffect();
-	}
-
-	if (UText3DComponent* Text3DComponent = GetText3DComponent())
-	{
-		Text3DComponent->OnTextGenerated().AddUObject(this, &UText3DCharacterTransform::ProcessEffect);
-	}
-}
-
-void UText3DCharacterTransform::OnUnregister()
-{
-	if (UText3DComponent* Text3DComponent = GetText3DComponent())
-	{
-		Text3DComponent->OnTextGenerated().RemoveAll(this);
-	}
-
-	Super::OnUnregister();
-}
-
-#if WITH_EDITOR
-void UText3DCharacterTransform::PostEditChangeProperty(FPropertyChangedEvent& PropertyChangedEvent)
-{
-	Super::PostEditChangeProperty(PropertyChangedEvent);
-
-	const FName Name = PropertyChangedEvent.GetPropertyName();
-	if (Name == GET_MEMBER_NAME_CHECKED(UText3DCharacterTransform, bLocationEnabled))
-	{
-		ResetLocation();
-	}
-	else if (Name == GET_MEMBER_NAME_CHECKED(UText3DCharacterTransform, bScaleEnabled))
-	{
-		ResetScale();
-	}
-	else if (Name == GET_MEMBER_NAME_CHECKED(UText3DCharacterTransform, bRotateEnabled))
-	{
-		ResetRotate();
-	}
-
-	ProcessEffect();
-}
-#endif
-
-void UText3DCharacterTransform::ProcessEffect()
-{
-	UText3DComponent* Text3DComponent = GetText3DComponent();
-	if (!Text3DComponent)
-	{
-		return;
-	}
-
-	// This will change to communicate directly with per CharacterData and renderer will not be used at all
-	UText3DStaticMeshesRenderer* TextRenderer = Cast<UText3DStaticMeshesRenderer>(Text3DComponent->GetTextRenderer());
-
-	int32 Count = TextRenderer->GetGlyphCount();
-	for (int32 Index = 0; Index < Count; Index++)
-	{
-		USceneComponent* GlyphComponent = TextRenderer->GetGlyphMeshComponent(Index);
-		if (!GlyphComponent)
-		{
-			continue;
-		}
-
-		if (bLocationEnabled)
-		{
-			float Effect = GetEffectValue(Index, Count, LocationOrder, LocationProgress, LocationRange);
-			FVector Location = LocationDistance * Effect;
-			GlyphComponent->SetRelativeLocation(Location);
-		}
-
-		if (bScaleEnabled)
-		{
-			float Effect = GetEffectValue(Index, Count, ScaleOrder, ScaleProgress, ScaleRange);
-			FVector Scale = ScaleBegin + ((ScaleEnd - ScaleBegin) * Effect);
-			GlyphComponent->SetRelativeScale3D(Scale);
-		}
-
-		if (bRotateEnabled)
-		{
-			float Effect = GetEffectValue(Index, Count, RotateOrder, RotateProgress, RotateRange);
-			FRotator Rotator = RotateBegin + (RotateEnd - RotateBegin) * Effect;
-			GlyphComponent->SetRelativeRotation(Rotator);
-		}
-	}
-}
-
-void UText3DCharacterTransform::ResetLocation()
-{
-	UText3DComponent* Text3DComponent = GetText3DComponent();
-	if (!Text3DComponent)
-	{
-		return;
-	}
-
-	// This will change to communicate directly with per CharacterData and renderer will not be used at all
-	UText3DStaticMeshesRenderer* TextRenderer = Cast<UText3DStaticMeshesRenderer>(Text3DComponent->GetTextRenderer());
-
-	int32 Count = TextRenderer->GetGlyphCount();
-	for (int32 Index = 0; Index < Count; Index++)
-	{
-		USceneComponent* GlyphComponent = TextRenderer->GetGlyphMeshComponent(Index);
-		if (GlyphComponent)
-		{
-			GlyphComponent->SetRelativeLocation(FVector::ZeroVector);
-		}
-	}
-}
-
-void UText3DCharacterTransform::ResetRotate()
-{
-	UText3DComponent* Text3DComponent = GetText3DComponent();
-	if (!Text3DComponent)
-	{
-		return;
-	}
-
-	// This will change to communicate directly with per CharacterData and renderer will not be used at all
-	UText3DStaticMeshesRenderer* TextRenderer = Cast<UText3DStaticMeshesRenderer>(Text3DComponent->GetTextRenderer());
-
-	int32 Count = TextRenderer->GetGlyphCount();
-	for (int32 Index = 0; Index < Count; Index++)
-	{
-		USceneComponent* GlyphComponent = TextRenderer->GetGlyphMeshComponent(Index);
-		if (GlyphComponent)
-		{
-			GlyphComponent->SetRelativeRotation(FRotator::ZeroRotator);
-		}
-	}
-}
-
-void UText3DCharacterTransform::ResetScale()
-{
-	UText3DComponent* Text3DComponent = GetText3DComponent();
-	if (!Text3DComponent)
-	{
-		return;
-	}
-
-	// This will change to communicate directly with per CharacterData and renderer will not be used at all
-	UText3DStaticMeshesRenderer* TextRenderer = Cast<UText3DStaticMeshesRenderer>(Text3DComponent->GetTextRenderer());
-
-	int32 Count = TextRenderer->GetGlyphCount();
-	for (int32 Index = 0; Index < Count; Index++)
-	{
-		USceneComponent* GlyphComponent = TextRenderer->GetGlyphMeshComponent(Index);
-		if (GlyphComponent)
-		{
-			GlyphComponent->SetRelativeScale3D(FVector(1.0f));
-		}
-	}
-}
-
-float UText3DCharacterTransform::GetEffectValue(int32 Index, int32 Total, EText3DCharacterEffectOrder Order, float Progress, float Range)
-{
-	float Effect = 0.0f, Strip = 1.0f;
-	GetLineParameters(Range * 0.01f, Order, Total, Effect, Strip);
-
-	int32 Position = GetEffectPosition(Index, Total, Order);
-	return FMath::Clamp((Progress * 0.01f - Effect * Position) / Strip, 0.0f, 1.0f);
-}
-
-void UText3DCharacterTransform::GetLineParameters(float Range, EText3DCharacterEffectOrder Order, int32 Total, float & EffectOut, float & StripOut)
-{
-	float Segment = 1.0f / Total;
-	EffectOut = (1.0f - Range) / Total;
-	StripOut = FMath::Max(Range, Segment);
-
-	if (Order == EText3DCharacterEffectOrder::FromCenter || Order == EText3DCharacterEffectOrder::ToCenter)
-	{
-		int32 Center = Total / 2.0f - 0.5f;
-		EffectOut *= (float)Total / (Center + 1.0f);
-	}
-
-	if (FMath::IsNearlyZero(EffectOut))
-	{
-		EffectOut = 0.01f;
-	}
-}
-
-int32 UText3DCharacterTransform::GetEffectPosition(int32 Index, int32 Total, EText3DCharacterEffectOrder Order)
-{
-	int32 Center = 0.5f * Total - 0.5f;
-	switch (Order)
-	{
-	case EText3DCharacterEffectOrder::FromCenter:
-	{
-		if (Index > Center)
-		{
-			Index = Total - Index - 1;
-		}
-
-		return Center - Index;
-	}
-
-	case EText3DCharacterEffectOrder::ToCenter:
-	{
-		if (Index > Center)
-		{
-			return Total - Index - 1;
-		}
-
-		return Index;
-	}
-
-	case EText3DCharacterEffectOrder::Opposite:
-	{
-		return Total - Index - 1;
-	}
-	}
-
-	return Index;
 }
 
 void UText3DCharacterTransform::SetLocationEnabled(bool bEnabled)
@@ -255,8 +35,11 @@ void UText3DCharacterTransform::SetLocationEnabled(bool bEnabled)
 	if (bLocationEnabled != bEnabled)
 	{
 		bLocationEnabled = bEnabled;
-		ResetLocation();
-		ProcessEffect();
+
+		if (UText3DLayoutTransformEffect* LayoutTransform = GetText3DLayoutTransformEffect())
+		{
+			LayoutTransform->SetLocationEnabled(bLocationEnabled);
+		}
 	}
 }
 
@@ -266,7 +49,11 @@ void UText3DCharacterTransform::SetLocationProgress(float Progress)
 	if (!FMath::IsNearlyEqual(LocationProgress, Progress))
 	{
 		LocationProgress = Progress;
-		ProcessEffect();
+
+		if (UText3DLayoutTransformEffect* LayoutTransform = GetText3DLayoutTransformEffect())
+		{
+			LayoutTransform->SetLocationProgress(LocationProgress);
+		}
 	}
 }
 
@@ -275,7 +62,11 @@ void UText3DCharacterTransform::SetLocationOrder(EText3DCharacterEffectOrder Ord
 	if (LocationOrder != Order)
 	{
 		LocationOrder = Order;
-		ProcessEffect();
+
+		if (UText3DLayoutTransformEffect* LayoutTransform = GetText3DLayoutTransformEffect())
+		{
+			LayoutTransform->SetLocationOrder(LocationOrder);
+		}
 	}
 }
 
@@ -285,16 +76,19 @@ void UText3DCharacterTransform::SetLocationRange(float Range)
 	if (!FMath::IsNearlyEqual(LocationRange, Range))
 	{
 		LocationRange = Range;
-		ProcessEffect();
 	}
 }
 
 void UText3DCharacterTransform::SetLocationDistance(FVector Distance)
 {
-	if (!(LocationDistance - Distance).IsNearlyZero())
+	if (!LocationDistance.Equals(Distance))
 	{
 		LocationDistance = Distance;
-		ProcessEffect();
+
+		if (UText3DLayoutTransformEffect* LayoutTransform = GetText3DLayoutTransformEffect())
+		{
+			LayoutTransform->SetLocationEnd(LocationDistance);
+		}
 	}
 }
 
@@ -303,8 +97,11 @@ void UText3DCharacterTransform::SetScaleEnabled(bool bEnabled)
 	if (bScaleEnabled != bEnabled)
 	{
 		bScaleEnabled = bEnabled;
-		ResetScale();
-		ProcessEffect();
+
+		if (UText3DLayoutTransformEffect* LayoutTransform = GetText3DLayoutTransformEffect())
+		{
+			LayoutTransform->SetScaleEnabled(bScaleEnabled);
+		}
 	}
 }
 
@@ -314,7 +111,11 @@ void UText3DCharacterTransform::SetScaleProgress(float Progress)
 	if (!FMath::IsNearlyEqual(ScaleProgress, Progress))
 	{
 		ScaleProgress = Progress;
-		ProcessEffect();
+
+		if (UText3DLayoutTransformEffect* LayoutTransform = GetText3DLayoutTransformEffect())
+		{
+			LayoutTransform->SetScaleProgress(ScaleProgress);
+		}
 	}
 }
 
@@ -323,7 +124,11 @@ void UText3DCharacterTransform::SetScaleOrder(EText3DCharacterEffectOrder Order)
 	if (ScaleOrder != Order)
 	{
 		ScaleOrder = Order;
-		ProcessEffect();
+
+		if (UText3DLayoutTransformEffect* LayoutTransform = GetText3DLayoutTransformEffect())
+		{
+			LayoutTransform->SetScaleOrder(ScaleOrder);
+		}
 	}
 }
 
@@ -333,25 +138,32 @@ void UText3DCharacterTransform::SetScaleRange(float Range)
 	if (!FMath::IsNearlyEqual(ScaleRange, Range))
 	{
 		ScaleRange = Range;
-		ProcessEffect();
 	}
 }
 
 void UText3DCharacterTransform::SetScaleBegin(FVector Value)
 {
-	if (!(ScaleBegin - Value).IsNearlyZero())
+	if (!ScaleBegin.Equals(Value))
 	{
 		ScaleBegin = Value;
-		ProcessEffect();
+
+		if (UText3DLayoutTransformEffect* LayoutTransform = GetText3DLayoutTransformEffect())
+		{
+			LayoutTransform->SetScaleBegin(ScaleBegin);
+		}
 	}
 }
 
 void UText3DCharacterTransform::SetScaleEnd(FVector Value)
 {
-	if (!(ScaleEnd - Value).IsNearlyZero())
+	if (!ScaleEnd.Equals(Value))
 	{
 		ScaleEnd = Value;
-		ProcessEffect();
+
+		if (UText3DLayoutTransformEffect* LayoutTransform = GetText3DLayoutTransformEffect())
+		{
+			LayoutTransform->SetScaleEnd(ScaleEnd);
+		}
 	}
 }
 
@@ -360,8 +172,11 @@ void UText3DCharacterTransform::SetRotateEnabled(bool bEnabled)
 	if (bRotateEnabled != bEnabled)
 	{
 		bRotateEnabled = bEnabled;
-		ResetRotate();
-		ProcessEffect();
+
+		if (UText3DLayoutTransformEffect* LayoutTransform = GetText3DLayoutTransformEffect())
+		{
+			LayoutTransform->SetRotationEnabled(bRotateEnabled);
+		}
 	}
 }
 
@@ -371,7 +186,11 @@ void UText3DCharacterTransform::SetRotateProgress(float Progress)
 	if (!FMath::IsNearlyEqual(RotateProgress, Progress))
 	{
 		RotateProgress = Progress;
-		ProcessEffect();
+
+		if (UText3DLayoutTransformEffect* LayoutTransform = GetText3DLayoutTransformEffect())
+		{
+			LayoutTransform->SetRotationProgress(RotateProgress);
+		}
 	}
 }
 
@@ -380,7 +199,11 @@ void UText3DCharacterTransform::SetRotateOrder(EText3DCharacterEffectOrder Order
 	if (RotateOrder != Order)
 	{
 		RotateOrder = Order;
-		ProcessEffect();
+
+		if (UText3DLayoutTransformEffect* LayoutTransform = GetText3DLayoutTransformEffect())
+		{
+			LayoutTransform->SetRotationOrder(RotateOrder);
+		}
 	}
 }
 
@@ -390,29 +213,36 @@ void UText3DCharacterTransform::SetRotateRange(float Range)
 	if (!FMath::IsNearlyEqual(RotateRange, Range))
 	{
 		RotateRange = Range;
-		ProcessEffect();
 	}
 }
 
 void UText3DCharacterTransform::SetRotateBegin(FRotator Value)
 {
-	if (!(RotateBegin - Value).IsNearlyZero())
+	if (!RotateBegin.Equals(Value))
 	{
 		RotateBegin = Value;
-		ProcessEffect();
+
+		if (UText3DLayoutTransformEffect* LayoutTransform = GetText3DLayoutTransformEffect())
+		{
+			LayoutTransform->SetRotationBegin(RotateBegin);
+		}
 	}
 }
 
 void UText3DCharacterTransform::SetRotateEnd(FRotator Value)
 {
-	if (!(RotateEnd - Value).IsNearlyZero())
+	if (!RotateEnd.Equals(Value))
 	{
 		RotateEnd = Value;
-		ProcessEffect();
+
+		if (UText3DLayoutTransformEffect* LayoutTransform = GetText3DLayoutTransformEffect())
+		{
+			LayoutTransform->SetRotationEnd(RotateEnd);
+		}
 	}
 }
 
-UText3DComponent* UText3DCharacterTransform::GetText3DComponent()
+UText3DComponent* UText3DCharacterTransform::GetText3DComponent() const
 {
 	UText3DComponent* Component = Cast<UText3DComponent>(GetAttachParent());
 	if (Component)
@@ -428,3 +258,40 @@ UText3DComponent* UText3DCharacterTransform::GetText3DComponent()
 	
 	return Owner->FindComponentByClass<UText3DComponent>();
 }
+
+UText3DLayoutTransformEffect* UText3DCharacterTransform::GetText3DLayoutTransformEffect() const
+{
+	if (const UText3DComponent* Text3DComponent = GetText3DComponent())
+	{
+		const TConstArrayView<UText3DLayoutEffectBase*> LayoutEffects = Text3DComponent->GetLayoutEffects();
+
+		if (!LayoutEffects.IsEmpty())
+		{
+			return Cast<UText3DLayoutTransformEffect>(LayoutEffects[0]);
+		}
+	}
+
+	return nullptr;
+}
+
+void UText3DCharacterTransform::OnComponentCreated()
+{
+	Super::OnComponentCreated();
+
+	PrintDeprecationLog();
+}
+
+void UText3DCharacterTransform::PostLoad()
+{
+	Super::PostLoad();
+
+	PrintDeprecationLog();
+}
+
+void UText3DCharacterTransform::PrintDeprecationLog() const
+{
+	if (const UText3DComponent* Text3DComponent = GetText3DComponent())
+	{
+		UE_LOG(LogText3D, Warning, TEXT("%s : Text3DCharacterTransform component is deprecated and should no longer be used, instead prefer using Text3DLayoutTransformEffect on Text3D that replaces this entirely !"), Text3DComponent->GetOwner() ? *Text3DComponent->GetOwner()->GetActorNameOrLabel() : TEXT("Invalid Actor"))
+	}
+}
diff --git a/Engine/Plugins/Experimental/Text3D/Source/Text3D/Public/Characters/Text3DCharacterTransform.h b/Engine/Plugins/Experimental/Text3D/Source/Text3D/Public/Characters/Text3DCharacterTransform.h
--- a/Engine/Plugins/Experimental/Text3D/Source/Text3D/Public/Characters/Text3DCharacterTransform.h
+++ b/Engine/Plugins/Experimental/Text3D/Source/Text3D/Public/Characters/Text3DCharacterTransform.h
@@ -6,9 +6,15 @@
 #include "Text3DTypes.h"
 #include "Text3DCharacterTransform.generated.h"
 
+class UText3DLayoutTransformEffect;
 class UText3DComponent;
 
-UCLASS(MinimalAPI, ClassGroup=(Text3D), HideCategories = (Collision, Tags, Activation, Cooking, Rendering, Physics, Mobility, LOD, AssetUserData, Navigation, Transform), Meta = (BlueprintSpawnableComponent))
+/**
+ * This class should no longer be used with Text3D
+ * It is only kept to migrate properties to the newer system
+ * See UText3DLayoutTransformEffect that can be added in section Text3DComponent/Layout/LayoutEffects
+ */
+UCLASS(MinimalAPI, ClassGroup=(Text3D), HideCategories = (Collision, Tags, Activation, Cooking, Rendering, Physics, Mobility, LOD, AssetUserData, Navigation, Transform))
 class UText3DCharacterTransform : public USceneComponent
 {
 	GENERATED_BODY()
@@ -159,28 +165,18 @@ public:
 	TEXT3D_API void SetScaleEnd(FVector Value);
 
 protected:
-	//~ Begin USceneComponent
-	virtual void OnRegister() override;
-	virtual void OnUnregister() override;
-	//~ End USceneComponent
+	UText3DComponent* GetText3DComponent() const;
+	UText3DLayoutTransformEffect* GetText3DLayoutTransformEffect() const;
+
+	//~ Begin UActorComponent
+	virtual void OnComponentCreated() override;
+	//~ End UActorComponent
 
 	//~ Begin UObject
-#if WITH_EDITOR
-	virtual void PostEditChangeProperty(struct FPropertyChangedEvent& PropertyChangedEvent) override;
-#endif
+	virtual void PostLoad() override;
 	//~ End UObject
 
-	UText3DComponent* GetText3DComponent();
-
-	void ProcessEffect();
-
-	float GetEffectValue(int32 Index, int32 Total, EText3DCharacterEffectOrder Order, float Progress, float Range);
-	void GetLineParameters(float Range, EText3DCharacterEffectOrder Order, int32 Count, float & EffectOut, float & StripOut);
-	int32 GetEffectPosition(int32 Index, int32 Total, EText3DCharacterEffectOrder Order);
-
-	void ResetLocation();
-	void ResetRotate();
-	void ResetScale();
+	void PrintDeprecationLog() const;
 
 	// Location
 
@@ -238,6 +234,4 @@ protected:
 
 	UPROPERTY(EditAnywhere, BlueprintReadWrite, BlueprintSetter = SetScaleEnd, Category = "Scale", meta = (EditCondition = "bScaleEnabled", ClampMin = 0))
 	FVector ScaleEnd;
-
-	bool bInitialized;
 };
