MotionDesign : Text3D
- Reworked font face and glyph mesh cache with schedule cleanup that can be configured in project settings
- Added rich text style formatting, define your style and use tags to use them
- Fix package marked dirty when text is updated due to component being detached
- Details panel category reordered and auto-expand transform and styles category

diff --git a/Engine/Plugins/Experimental/Text3D/Source/Text3D/Private/Characters/Text3DCharacterBase.cpp b/Engine/Plugins/Experimental/Text3D/Source/Text3D/Private/Characters/Text3DCharacterBase.cpp
--- a/Engine/Plugins/Experimental/Text3D/Source/Text3D/Private/Characters/Text3DCharacterBase.cpp
+++ b/Engine/Plugins/Experimental/Text3D/Source/Text3D/Private/Characters/Text3DCharacterBase.cpp
@@ -34,34 +34,38 @@ FTransform& UText3DCharacterBase::GetTransform(bool bInReset)
 	return Transform;
 }
 
-void UText3DCharacterBase::SetGlyphIndex(uint32 InGlyphIndex)
+const UE::Text3D::Geometry::FCachedFontFaceGlyphHandle& UText3DCharacterBase::GetFontFaceGlyphHandle() const
 {
-	GlyphIndex = InGlyphIndex;
+	return FontFaceGlyphHandle;
 }
 
-uint32 UText3DCharacterBase::GetGlyphIndex() const
+void UText3DCharacterBase::SetFontFaceGlyphHandle(const UE::Text3D::Geometry::FCachedFontFaceGlyphHandle& InHandle)
 {
-	return GlyphIndex;
+	FontFaceGlyphHandle = InHandle;
 }
 
-void UText3DCharacterBase::SetMeshBounds(const FBox& InBounds)
-{
-	MeshBounds = InBounds;
-}
-
-const FBox& UText3DCharacterBase::GetMeshBounds() const
+uint32 UText3DCharacterBase::GetGlyphIndex() const
 {
-	return MeshBounds;
+	return FontFaceGlyphHandle.GetGlyphIndex();
 }
 
-void UText3DCharacterBase::SetMeshOffset(const FVector& InOffset)
+const FText3DCachedMesh* UText3DCharacterBase::GetGlyphMesh() const
 {
-	MeshOffset = InOffset;
+	return FontFaceGlyphHandle.Resolve();
 }
 
-const FVector& UText3DCharacterBase::GetMeshOffset() const
+void UText3DCharacterBase::GetGlyphMeshBoundsAndOffset(FBox& OutBounds, FVector& OutOffset) const
 {
-	return MeshOffset;
+	if (const FText3DCachedMesh* GlyphMesh = GetGlyphMesh())
+	{
+		OutBounds = GlyphMesh->MeshBounds;
+		OutOffset = GlyphMesh->MeshOffset;
+	}
+	else
+	{
+		OutBounds = FBox(EForceInit::ForceInitToZero);
+		OutOffset = FVector::ZeroVector;
+	}
 }
 
 void UText3DCharacterBase::SetRelativeLocation(const FVector& InLocation)
@@ -119,9 +123,7 @@ void UText3DCharacterBase::ResetCharacterState()
 	RelativeScale = CDO->RelativeScale;
 	bVisible = CDO->bVisible;
 	Transform = CDO->Transform;
-	GlyphIndex = CDO->GlyphIndex;
-	MeshBounds = CDO->MeshBounds;
-	MeshOffset = CDO->MeshOffset;
+	FontFaceGlyphHandle.Unset();
 }
 
 #if WITH_EDITOR
diff --git a/Engine/Plugins/Experimental/Text3D/Source/Text3D/Private/Extensions/Text3DDefaultFormatStyleExtension.cpp b/Engine/Plugins/Experimental/Text3D/Source/Text3D/Private/Extensions/Text3DDefaultFormatStyleExtension.cpp
new file mode 100644
--- /dev/null
+++ b/Engine/Plugins/Experimental/Text3D/Source/Text3D/Private/Extensions/Text3DDefaultFormatStyleExtension.cpp
@@ -0,0 +1,76 @@
+// Copyright Epic Games, Inc. All Rights Reserved.
+
+#include "Extensions/Text3DDefaultFormatStyleExtension.h"
+
+#include "Engine/Font.h"
+#include "Fonts/SlateFontInfo.h"
+#include "Styles/Text3DFormatStyleBase.h"
+#include "Styling/SlateStyle.h"
+#include "Text3DComponent.h"
+
+#if WITH_EDITOR
+void UText3DDefaultFormatStyleExtension::PostEditChangeProperty(FPropertyChangedEvent& InPropertyChangedEvent)
+{
+	Super::PostEditChangeProperty(InPropertyChangedEvent);
+
+	const FName MemberPropertyName = InPropertyChangedEvent.GetMemberPropertyName();
+
+	if (MemberPropertyName == GET_MEMBER_NAME_CHECKED(UText3DDefaultFormatStyleExtension, FormatStyles))
+	{
+		if (InPropertyChangedEvent.ChangeType == EPropertyChangeType::ArrayClear
+			|| InPropertyChangedEvent.ChangeType == EPropertyChangeType::ArrayRemove)
+		{
+			if (UText3DComponent* Text3DComponent = GetText3DComponent())
+			{
+				Text3DComponent->RequestUpdate(EText3DRendererFlags::All, /** Immediate */false);
+			}
+		}
+	}
+}
+#endif
+
+TSharedPtr<FTextBlockStyle> UText3DDefaultFormatStyleExtension::GetDefaultStyle() const
+{
+	return DefaultFontStyle;
+}
+
+TSharedPtr<FSlateStyleSet> UText3DDefaultFormatStyleExtension::GetCustomStyles() const
+{
+	return TextStyleSet;
+}
+
+EText3DExtensionResult UText3DDefaultFormatStyleExtension::PreRendererUpdate(const UE::Text3D::Renderer::FUpdateParameters& InParameters)
+{
+	if (!DefaultFontStyle.IsValid())
+	{
+		DefaultFontStyle = MakeShared<FTextBlockStyle>();
+	}
+
+	if (EnumHasAnyFlags(InParameters.UpdateFlags, EText3DRendererFlags::Geometry))
+	{
+		const UText3DComponent* Text3DComponent = GetText3DComponent();
+		const UFont* Font = Text3DComponent->GetFont();
+		FSlateFontInfo FontInfo(Font, UE::Text3D::Metrics::SlateFontSize);
+		FontInfo.CompositeFont = Font && Font->GetCompositeFont() ? MakeShared<FCompositeFont>(*Font->GetCompositeFont()) : FStyleDefaults::GetFontInfo().CompositeFont;
+		FontInfo.TypefaceFontName = Text3DComponent->GetTypeface();
+		FontInfo.FontFallback = EFontFallback::FF_NoFallback;
+		DefaultFontStyle->SetFont(FontInfo);
+
+		TextStyleSet = MakeShared<FSlateStyleSet>(TEXT("Text3DRichTextStyle"));
+		TextStyleSet->Set(FName(TEXT("Default")), *DefaultFontStyle);
+		for (const UText3DFormatStyleBase* FormatStyle : FormatStyles)
+		{
+			if (FormatStyle && !FormatStyle->GetStyleName().IsNone())
+			{
+				TextStyleSet->Set(FormatStyle->GetStyleName(), FormatStyle->GetTextStyle());
+			}
+		}
+	}
+
+	return EText3DExtensionResult::Finished;
+}
+
+EText3DExtensionResult UText3DDefaultFormatStyleExtension::PostRendererUpdate(const UE::Text3D::Renderer::FUpdateParameters& InParameters)
+{
+	return EText3DExtensionResult::Active;
+}
diff --git a/Engine/Plugins/Experimental/Text3D/Source/Text3D/Private/Extensions/Text3DDefaultGeometryExtension.cpp b/Engine/Plugins/Experimental/Text3D/Source/Text3D/Private/Extensions/Text3DDefaultGeometryExtension.cpp
--- a/Engine/Plugins/Experimental/Text3D/Source/Text3D/Private/Extensions/Text3DDefaultGeometryExtension.cpp
+++ b/Engine/Plugins/Experimental/Text3D/Source/Text3D/Private/Extensions/Text3DDefaultGeometryExtension.cpp
@@ -151,14 +151,7 @@ void UText3DDefaultGeometryExtension::PostEditChangeProperty(FPropertyChangedEve
 
 const FText3DCachedMesh* UText3DDefaultGeometryExtension::FindOrLoadGlyphMesh(UText3DCharacterBase* InCharacter) const
 {
-	if (const FText3DCachedMesh* GlyphMesh = CachedFontData->GetGlyphMesh(InCharacter->GetGlyphIndex(), GlyphMeshParameters))
-	{
-		InCharacter->SetMeshBounds(GlyphMesh->MeshBounds);
-		InCharacter->SetMeshOffset(GlyphMesh->MeshOffset);
-		return GlyphMesh;
-	}
-
-	return nullptr;
+	return InCharacter->GetGlyphMesh();
 }
 
 EText3DHorizontalTextAlignment UText3DDefaultGeometryExtension::GetGlyphHAlignment() const
@@ -171,14 +164,9 @@ EText3DVerticalTextAlignment UText3DDefaultGeometryExtension::GetGlyphVAlignment
 	return PivotVAlignment;
 }
 
-FTypefaceFontData* UText3DDefaultGeometryExtension::GetTypefaceFontData() const
+const FGlyphMeshParameters* UText3DDefaultGeometryExtension::GetGlyphMeshParameters() const
 {
-	FTypefaceFontData* TypefaceFontData = nullptr;
-	if (CachedFontData)
-	{
-		TypefaceFontData = CachedFontData->GetTypefaceFontData(GlyphMeshParameters.TypefaceIndex);
-	}
-	return TypefaceFontData;
+	return &GlyphMeshParameters;
 }
 
 EText3DExtensionResult UText3DDefaultGeometryExtension::PreRendererUpdate(const UE::Text3D::Renderer::FUpdateParameters& InParameters)
@@ -189,8 +177,8 @@ EText3DExtensionResult UText3DDefaultGeometryExtension::PreRendererUpdate(const
 	}
 
 	const UText3DComponent* TextComponent = GetText3DComponent();
-	const uint32 TypefaceIndex = TextComponent->GetTypeFaceIndex();
-	UFont* Font = TextComponent->GetFont();
+	const int32 TypefaceIndex = TextComponent->GetTypeFaceIndex();
+	const UFont* Font = TextComponent->GetFont();
 
 	if (TypefaceIndex == INDEX_NONE)
 	{
@@ -198,25 +186,19 @@ EText3DExtensionResult UText3DDefaultGeometryExtension::PreRendererUpdate(const
 		return EText3DExtensionResult::Failed;
 	}
 
-	UText3DEngineSubsystem* Subsystem = UText3DEngineSubsystem::Get();
-	CachedFontData = &Subsystem->GetCachedFontData(Font, TypefaceIndex);
-
-	if (!CachedFontData)
-	{
-		UE_LOG(LogText3D, Error, TEXT("Failed to get cached font data for '%s %i' in Text3D geometry extension"), Font ? *Font->GetName() : TEXT("invalid"), TypefaceIndex);
-		return EText3DExtensionResult::Failed;
-	}
+	const float GlyphExtrude = FMath::TruncToFloat(GetExtrude() * 10.f) / 10.f;
+	const float GlyphBevel = FMath::TruncToFloat(GetBevel() * 10.f) / 10.f;
+	const float GlyphOutline = FMath::TruncToFloat(GetOutline() * 10.f) / 10.f;
 
 	GlyphMeshParameters =
 	{
-		.Extrude = GetExtrude(),
-		.Bevel = GetBevel(),
+		.Extrude = GlyphExtrude,
+		.Bevel = GlyphBevel,
 		.BevelType = GetBevelType(),
 		.BevelSegments = GetBevelSegments(),
 		.bOutline = GetUseOutline(),
-		.OutlineExpand = GetOutline(),
+		.OutlineExpand = GlyphOutline,
 		.OutlineType = GetOutlineType(),
-		.TypefaceIndex = TypefaceIndex,
 		.PivotOffset = GetPivotOffset(),
 	};
 
@@ -226,7 +208,6 @@ EText3DExtensionResult UText3DDefaultGeometryExtension::PreRendererUpdate(const
 
 EText3DExtensionResult UText3DDefaultGeometryExtension::PostRendererUpdate(const UE::Text3D::Renderer::FUpdateParameters& InParameters)
 {
-	CachedFontData = nullptr;
 	return EText3DExtensionResult::Finished;
 }
 
diff --git a/Engine/Plugins/Experimental/Text3D/Source/Text3D/Private/Extensions/Text3DDefaultLayoutExtension.cpp b/Engine/Plugins/Experimental/Text3D/Source/Text3D/Private/Extensions/Text3DDefaultLayoutExtension.cpp
--- a/Engine/Plugins/Experimental/Text3D/Source/Text3D/Private/Extensions/Text3DDefaultLayoutExtension.cpp
+++ b/Engine/Plugins/Experimental/Text3D/Source/Text3D/Private/Extensions/Text3DDefaultLayoutExtension.cpp
@@ -7,8 +7,10 @@
 #include "Characters/Text3DCharacterBase.h"
 #include "Engine/Font.h"
 #include "Extensions/Text3DCharacterExtensionBase.h"
+#include "Extensions/Text3DFormatStyleExtensionBase.h"
 #include "Extensions/Text3DGeometryExtensionBase.h"
-#include "Framework/Text/PlainTextLayoutMarshaller.h"
+#include "Framework/Text/RichTextLayoutMarshaller.h"
+#include "Framework/Text/RichTextMarkupProcessing.h"
 #include "LayoutBuilders/Text3DLayout.h"
 #include "LayoutBuilders/Text3DLayoutShaper.h"
 #include "LayoutBuilders/Text3DShapedGlyphText.h"
@@ -173,7 +175,14 @@ void UText3DDefaultLayoutExtension::SetScaleProportionally(const bool Value)
 
 float UText3DDefaultLayoutExtension::GetTextHeight() const
 {
-	return ShapedText->Lines.Num() * TypeFaceMetrics.FontHeight + (ShapedText->Lines.Num() - 1) * LineSpacing;
+	float Height = FMath::Max(ShapedText->Lines.Num() - 1, 0) * LineSpacing;
+
+	for (const FShapedGlyphLine& Line : ShapedText->Lines)
+	{
+		Height += Line.MaxFontHeight;
+	}
+
+	return Height;
 }
 
 FVector UText3DDefaultLayoutExtension::GetTextScale() const
@@ -188,7 +197,7 @@ EText3DExtensionResult UText3DDefaultLayoutExtension::PreRendererUpdate(const UE
 		return EText3DExtensionResult::Active;
 	}
 
-	UText3DComponent* Text3DComponent = GetText3DComponent();
+	const UText3DComponent* Text3DComponent = GetText3DComponent();
 	
 	if (!ShapedText.IsValid())
 	{
@@ -202,31 +211,20 @@ EText3DExtensionResult UText3DDefaultLayoutExtension::PreRendererUpdate(const UE
 
 	if (!TextLayoutMarshaller.IsValid())
 	{
-		TextLayoutMarshaller = FPlainTextLayoutMarshaller::Create();
-	}
-
-	if (!FontStyle.IsValid())
-	{
-		FontStyle = MakeShared<FTextBlockStyle>();
-	}
-
-	if (EnumHasAnyFlags(InParameters.UpdateFlags, EText3DRendererFlags::Geometry))
-	{
-		UFont* Font = Text3DComponent->GetFont();
-		FSlateFontInfo FontInfo(Font, UE::Text3D::Metrics::SlateFontSize);
-		FontInfo.CompositeFont = Font->GetCompositeFont() ? MakeShared<FCompositeFont>(*Text3DComponent->GetFont()->GetCompositeFont()) : FStyleDefaults::GetFontInfo().CompositeFont;
-		FontInfo.TypefaceFontName = Text3DComponent->GetTypeface();
-		FontStyle->SetFont(FontInfo);
+		const TSharedPtr<IRichTextMarkupParser> Parser = FDefaultRichTextMarkupParser::GetStaticInstance();
+		TextLayoutMarshaller = FRichTextLayoutMarshaller::Create(Parser, nullptr, TArray<TSharedRef<ITextDecorator>>(), nullptr);
 	}
 
 	ShapedText->Reset();
-	FText3DLayoutShaper::Get()->ShapeBidirectionalText(*FontStyle, Text3DComponent->GetFormattedText().ToString(), TextLayout, TextLayoutMarshaller, ShapedText->Lines);
+
+	const UText3DFormatStyleExtensionBase* FormatStyleExtension = Text3DComponent->GetFormatStyleExtension();
+	FText3DLayoutShaper::Get()->ShapeBidirectionalText(FormatStyleExtension->GetCustomStyles(), *FormatStyleExtension->GetDefaultStyle(), Text3DComponent->GetFormattedText().ToString(), TextLayout, TextLayoutMarshaller, ShapedText->Lines);
 
 	const int32 CharacterCount = Algo::TransformAccumulate(ShapedText->Lines, [](const FShapedGlyphLine& InLine)
 	{
 		return Algo::CountIf(InLine.GlyphsToRender, [](const FShapedGlyphEntry& InGlyph)
 		{
-			return InGlyph.HasValidGlyph() && InGlyph.bIsVisible;
+			return FText3DShapedGlyphText::IsValidGlyph(InGlyph);
 		});
 	},0);
 
@@ -236,65 +234,30 @@ EText3DExtensionResult UText3DDefaultLayoutExtension::PreRendererUpdate(const UE
 		CharacterExtension->AllocateCharacters(CharacterCount);
 	}
 
-	ShapedText->GlyphIndexToFontFace.Reserve(CharacterCount);
 	ShapedText->Kernings.SetNum(CharacterCount);
+	ShapedText->FontFaces.SetNum(CharacterCount);
 
-	UText3DGeometryExtensionBase* GeometryExtension = Text3DComponent->GetGeometryExtension();
+	const UText3DGeometryExtensionBase* GeometryExtension = Text3DComponent->GetGeometryExtension();
 	uint16 CharacterIndex = 0;
 
 #if WITH_FREETYPE
-	FT_Face FontFace = nullptr;
-	for (TConstEnumerateRef<FShapedGlyphLine> ShapedLine : EnumerateRange(ShapedText->Lines))
+	for (const TConstEnumerateRef<FShapedGlyphLine> ShapedLine : EnumerateRange(ShapedText->Lines))
 	{
 		const bool bIsRTL = (ShapedLine->TextDirection == TextBiDi::ETextDirection::RightToLeft);
 
-		for (TConstEnumerateRef<FShapedGlyphEntry> GlyphEntry : EnumerateRange(ShapedLine->GlyphsToRender))
+		for (const TConstEnumerateRef<FShapedGlyphEntry> GlyphEntry : EnumerateRange(ShapedLine->GlyphsToRender))
 		{
-			if (!GlyphEntry->FontFaceData.IsValid())
-			{
-				// Add as nullptr if not already in map
-				if (!ShapedText->GlyphIndexToFontFace.Contains(GlyphEntry->GlyphIndex))
-				{
-					ShapedText->GlyphIndexToFontFace.Add(GlyphEntry->GlyphIndex, nullptr);
-				}
-
-				continue;
-			}
-
-			if (const TSharedPtr<FFreeTypeFace> FontFacePtr = GlyphEntry->FontFaceData->FontFace.Pin())
-			{
-				ShapedText->GlyphIndexToFontFace.FindOrAdd(GlyphEntry->GlyphIndex, FontFacePtr.Get());
-				FontFace = FontFacePtr->GetFace();
-			}
-
-			if (GlyphEntry->HasValidGlyph() && GlyphEntry->bIsVisible)
+			if (FText3DShapedGlyphText::IsValidGlyph(*GlyphEntry))
 			{
+				const TSharedPtr<FFreeTypeFace> FontFacePtr = GlyphEntry->FontFaceData->FontFace.Pin();
 				const uint16 Index = bIsRTL ? CharacterCount - 1 - CharacterIndex : CharacterIndex;
-				UText3DCharacterBase* Character = Text3DComponent->GetCharacter(Index);
+				const UText3DCharacterBase* Character = Text3DComponent->GetCharacter(Index);
 				ShapedText->Kernings[CharacterIndex] = Character->GetCharacterKerning();
+				ShapedText->FontFaces[CharacterIndex] = FontFacePtr;
 				CharacterIndex++;
 			}
 		}
 	}
-
-	if (EnumHasAnyFlags(InParameters.UpdateFlags, EText3DRendererFlags::Geometry))
-	{
-		FTypefaceFontData* TypefaceFontData = GeometryExtension->GetTypefaceFontData();
-		if (!TypefaceFontData)
-		{
-			UE_LOG(LogText3D, Warning, TEXT("Typeface font data cannot be retrieved, cannot proceed without it"))
-			return EText3DExtensionResult::Failed;
-		}
-
-		TypefaceFontData->SetTypeFace(FontFace);
-
-		if (FontFace)
-		{
-			TypeFaceMetrics.FontHeight = FontFace->size->metrics.height * UE::Text3D::Metrics::FontSizeInverse;
-			TypeFaceMetrics.FontAscender = FontFace->size->metrics.ascender * UE::Text3D::Metrics::FontSizeInverse;
-			TypeFaceMetrics.FontDescender = FontFace->size->metrics.descender * UE::Text3D::Metrics::FontSizeInverse;
-		}
-	}
 #else
 	UE_LOG(LogText3D, Warning, TEXT("FreeType is not available, cannot proceed without it"))
 	return EText3DExtensionResult::Failed;
@@ -308,6 +271,7 @@ EText3DExtensionResult UText3DDefaultLayoutExtension::PreRendererUpdate(const UE
 	CalculateTextScale();
 
 	CharacterIndex = 0;
+	UText3DEngineSubsystem* TextSubsystem = UText3DEngineSubsystem::Get();
 	for (TConstEnumerateRef<FShapedGlyphLine> ShapedLine : EnumerateRange(ShapedText->Lines))
 	{
 		FVector Location = GetLineLocation(ShapedLine.GetIndex());
@@ -315,21 +279,44 @@ EText3DExtensionResult UText3DDefaultLayoutExtension::PreRendererUpdate(const UE
 
 		for (TConstEnumerateRef<FShapedGlyphEntry> ShapedGlyph : EnumerateRange(ShapedLine->GlyphsToRender))
 		{
-			if (ShapedGlyph->HasValidGlyph() && ShapedGlyph->bIsVisible)
+			if (FText3DShapedGlyphText::IsValidGlyph(*ShapedGlyph))
 			{
 				const uint16 Index = bIsRTL ? CharacterCount - 1 - CharacterIndex : CharacterIndex;
 				UText3DCharacterBase* Character = Text3DComponent->GetCharacter(Index);
 
 				if (EnumHasAnyFlags(InParameters.UpdateFlags, EText3DRendererFlags::Geometry))
 				{
-					Character->SetGlyphIndex(ShapedGlyph->GlyphIndex);
-					GeometryExtension->FindOrLoadGlyphMesh(Character);
+#if WITH_FREETYPE
+					TSharedPtr<FFreeTypeFace> CharTypeface = ShapedText->FontFaces[CharacterIndex];
+					uint32 FontFaceHash = FText3DFontFaceCache::GetFontFaceHash(CharTypeface);
+					// Cache font face to keep it alive and use inner FT_Face to build glyphs
+					FText3DFontFaceCache* FontFaceCache = TextSubsystem->FindOrAddCachedFontFace(CharTypeface);
+#else
+					uint32 FontFaceHash = 0;
+					FText3DFontFaceCache* FontFaceCache = nullptr;
+#endif
+
+					if (FontFaceCache)
+					{
+						// Load glyph here to get proper mesh metrics
+						UE::Text3D::Geometry::FCachedFontFaceGlyphHandle GlyphHandle(FontFaceCache, ShapedGlyph->GlyphIndex, *GeometryExtension->GetGlyphMeshParameters());
+						Character->SetFontFaceGlyphHandle(GlyphHandle);
+					}
+					else
+					{
+						UE_LOG(LogText3D, Error, TEXT("Failed to get cached font data for '%u %i' in Text3D geometry extension"), FontFaceHash, ShapedGlyph->GlyphIndex);
+					}
 				}
 
 				constexpr bool bReset = true;
 				FTransform& CharacterTransform = Character->GetTransform(bReset);
 
-				const FVector GlyphSize = Character->GetMeshBounds().GetSize();
+				FBox GlyphBounds;
+				FVector GlyphOffset;
+				Character->GetGlyphMeshBoundsAndOffset(GlyphBounds, GlyphOffset);
+
+				const FVector GlyphSize = GlyphBounds.GetSize();
+
 				FVector Shift = FVector::ZeroVector;
 				switch (GeometryExtension->GetGlyphHAlignment())
 				{
@@ -406,6 +393,7 @@ FVector UText3DDefaultLayoutExtension::GetLineLocation(int32 LineIndex) const
 {
 	float HorizontalOffset = 0.0f;
 	float VerticalOffset = 0.0f;
+
 	if (LineIndex < 0 || LineIndex >= ShapedText->Lines.Num())
 	{
 		return FVector::ZeroVector;
@@ -426,7 +414,7 @@ FVector UText3DDefaultLayoutExtension::GetLineLocation(int32 LineIndex) const
 	if (VerticalAlignment != EText3DVerticalTextAlignment::FirstLine)
 	{
 		// First align it to Top
-		VerticalOffset -= TypeFaceMetrics.FontAscender;
+		VerticalOffset -= ShapedLine.MaxFontAscender;
 
 		if (VerticalAlignment == EText3DVerticalTextAlignment::Center)
 		{
@@ -434,11 +422,16 @@ FVector UText3DDefaultLayoutExtension::GetLineLocation(int32 LineIndex) const
 		}
 		else if (VerticalAlignment == EText3DVerticalTextAlignment::Bottom)
 		{
-			VerticalOffset += TotalHeight + TypeFaceMetrics.FontDescender;
+			VerticalOffset += TotalHeight + ShapedLine.MaxFontDescender;
 		}
 	}
 
-	VerticalOffset -= LineIndex * (TypeFaceMetrics.FontHeight + LineSpacing);
+	// Offset line based on cumulative height above it
+	for (int32 Index = 0; Index < LineIndex; ++Index)
+	{
+		const FShapedGlyphLine& PrevLine = ShapedText->Lines[Index];
+		VerticalOffset -= (PrevLine.MaxFontHeight + LineSpacing);
+	}
 
 	return FVector(0.0f, HorizontalOffset, VerticalOffset);
 }
diff --git a/Engine/Plugins/Experimental/Text3D/Source/Text3D/Private/LayoutBuilders/Text3DLayoutShaper.cpp b/Engine/Plugins/Experimental/Text3D/Source/Text3D/Private/LayoutBuilders/Text3DLayoutShaper.cpp
--- a/Engine/Plugins/Experimental/Text3D/Source/Text3D/Private/LayoutBuilders/Text3DLayoutShaper.cpp
+++ b/Engine/Plugins/Experimental/Text3D/Source/Text3D/Private/LayoutBuilders/Text3DLayoutShaper.cpp
@@ -2,17 +2,17 @@
 
 #include "LayoutBuilders/Text3DLayoutShaper.h"
 
-#include "Brushes/SlateNoResource.h"
 #include "Containers/Array.h"
 #include "Containers/UnrealString.h"
 #include "Fonts/FontCache.h"
 #include "Framework/Application/SlateApplication.h"
-#include "Framework/Text/PlainTextLayoutMarshaller.h"
+#include "Framework/Text/RichTextLayoutMarshaller.h"
 #include "Framework/Text/ShapedTextCache.h"
 #include "Framework/Text/SlateTextRun.h"
 #include "Framework/Text/TextLayout.h"
 #include "Internationalization/Text.h"
 #include "Text3DShapedGlyphLine.h"
+#include "Styling/SlateStyle.h"
 
 TSharedPtr<FText3DLayoutShaper> FText3DLayoutShaper::Get()
 {
@@ -21,10 +21,11 @@ TSharedPtr<FText3DLayoutShaper> FText3DLayoutShaper::Get()
 }
 
 void FText3DLayoutShaper::ShapeBidirectionalText(
-	const FTextBlockStyle& Style,
+	const TSharedPtr<FSlateStyleSet>& InStyles,
+	const FTextBlockStyle& InDefaultStyle,
 	const FString& Text,
 	const TSharedPtr<FTextLayout>& TextLayout,
-	const TSharedPtr<ITextLayoutMarshaller>& TextMarshaller,
+	const TSharedPtr<FRichTextLayoutMarshaller>& TextMarshaller,
 	TArray<FShapedGlyphLine>& OutShapedLines)
 {
 	// @todo: Restore when dependency on SlateApplication is removed. Currently this means text meshes won't be created on the server.
@@ -37,7 +38,8 @@ void FText3DLayoutShaper::ShapeBidirectionalText(
 	TextLayout->ClearLines();
 	TextLayout->ClearLineHighlights();
 	TextLayout->ClearRunRenderers();
-	
+
+	TextMarshaller->SetDecoratorStyleSet(InStyles.Get());
 	TextMarshaller->SetText(Text, *TextLayout);
 	TextMarshaller->ClearDirty();
 
@@ -55,17 +57,17 @@ void FText3DLayoutShaper::ShapeBidirectionalText(
 		{
 			FString BlockText;
 			Block->GetRun()->AppendTextTo(BlockText);
-
-			TSharedRef<const FString> BlockTextRef = MakeShared<const FString>(BlockText);
-			
-			const FTextRange BlockRange = Block->GetTextRange();
+	
 			const FLayoutBlockTextContext BlockTextContext = Block->GetTextContext();
+			const FName StyleTag(Block->GetRun()->GetRunInfo().Name);
+			const FSlateFontInfo& FontInfo = !StyleTag.IsNone() && InStyles->HasWidgetStyle<FTextBlockStyle>(StyleTag) ? InStyles->GetWidgetStyle<FTextBlockStyle>(StyleTag).Font : InDefaultStyle.Font;
+			const FTextRange BlockRange(0, BlockText.Len());
 
 			const FShapedGlyphSequenceRef ShapedGlyphSequence = ShapedTextCacheUtil::GetShapedTextSubSequence(
 				BlockTextContext.ShapedTextCache,
-				FCachedShapedTextKey(Line.Range, 1.0f, BlockTextContext, Style.Font),
+				FCachedShapedTextKey(BlockRange, 1.0f, BlockTextContext, FontInfo),
 				BlockRange,
-				**BlockTextRef,
+				*BlockText,
 				BlockTextContext.TextDirection
 			);
 
diff --git a/Engine/Plugins/Experimental/Text3D/Source/Text3D/Private/LayoutBuilders/Text3DLayoutShaper.h b/Engine/Plugins/Experimental/Text3D/Source/Text3D/Private/LayoutBuilders/Text3DLayoutShaper.h
--- a/Engine/Plugins/Experimental/Text3D/Source/Text3D/Private/LayoutBuilders/Text3DLayoutShaper.h
+++ b/Engine/Plugins/Experimental/Text3D/Source/Text3D/Private/LayoutBuilders/Text3DLayoutShaper.h
@@ -6,8 +6,9 @@
 #include "Internationalization/Text.h"
 #include "Templates/SharedPointer.h"
 
+class FRichTextLayoutMarshaller;
+class FSlateStyleSet;
 class FTextLayout;
-class ITextLayoutMarshaller;
 struct FShapedGlyphLine;
 struct FTextBlockStyle;
 
@@ -23,10 +24,11 @@ public:
 	 * Analogous to FSlateFontCache::ShapeBidirectionalText.
 	 */
 	void ShapeBidirectionalText(
-		const FTextBlockStyle& Style,
+		const TSharedPtr<FSlateStyleSet>& InStyles,
+		const FTextBlockStyle& InDefaultStyle,
 		const FString& Text,
 		const TSharedPtr<FTextLayout>& TextLayout,
-		const TSharedPtr<ITextLayoutMarshaller>& TextMarshaller,
+		const TSharedPtr<FRichTextLayoutMarshaller>& TextMarshaller,
 		TArray<FShapedGlyphLine>& OutShapedLines);
 
 private:
diff --git a/Engine/Plugins/Experimental/Text3D/Source/Text3D/Private/LayoutBuilders/Text3DShapedGlyphLine.h b/Engine/Plugins/Experimental/Text3D/Source/Text3D/Private/LayoutBuilders/Text3DShapedGlyphLine.h
--- a/Engine/Plugins/Experimental/Text3D/Source/Text3D/Private/LayoutBuilders/Text3DShapedGlyphLine.h
+++ b/Engine/Plugins/Experimental/Text3D/Source/Text3D/Private/LayoutBuilders/Text3DShapedGlyphLine.h
@@ -23,4 +23,13 @@ struct FShapedGlyphLine
 
 	/** Direction of the text in the line */
 	TextBiDi::ETextDirection TextDirection = TextBiDi::ETextDirection::LeftToRight;
+
+	/** Max font height for this line */
+	float MaxFontHeight = 0.f;
+
+	/** Max height above baseline */
+	float MaxFontAscender = 0.f;
+
+	/** Max height below baseline */
+	float MaxFontDescender = 0.f;
 };
diff --git a/Engine/Plugins/Experimental/Text3D/Source/Text3D/Private/LayoutBuilders/Text3DShapedGlyphText.cpp b/Engine/Plugins/Experimental/Text3D/Source/Text3D/Private/LayoutBuilders/Text3DShapedGlyphText.cpp
--- a/Engine/Plugins/Experimental/Text3D/Source/Text3D/Private/LayoutBuilders/Text3DShapedGlyphText.cpp
+++ b/Engine/Plugins/Experimental/Text3D/Source/Text3D/Private/LayoutBuilders/Text3DShapedGlyphText.cpp
@@ -1,9 +1,21 @@
 // Copyright Epic Games, Inc. All Rights Reserved.
 
 #include "LayoutBuilders/Text3DShapedGlyphText.h"
-
 #include "Misc/EnumerateRange.h"
 
+#if WITH_FREETYPE
+#include "Fonts/SlateTextShaper.h"
+#endif
+
+bool FText3DShapedGlyphText::IsValidGlyph(const FShapedGlyphEntry& InGlyph)
+{
+#if WITH_FREETYPE
+	return InGlyph.HasValidGlyph() && InGlyph.bIsVisible && InGlyph.GlyphIndex != 0 && InGlyph.FontFaceData.IsValid() && InGlyph.FontFaceData->FontFace.IsValid();
+#else
+	return false;
+#endif
+}
+
 FText3DShapedGlyphText::FText3DShapedGlyphText()
 {
 	Reset();
@@ -12,11 +24,11 @@ FText3DShapedGlyphText::FText3DShapedGlyphText()
 void FText3DShapedGlyphText::Reset()
 {
 	Kernings.Reset();
+	FontFaces.Reset();
 	Tracking = 0.f;
 	WordSpacing = 0.0f;
 	bWrap = false;
 	Lines.Reset();
-	GlyphIndexToFontFace.Reset();
 }
 
 void FText3DShapedGlyphText::CalculateWidth()
@@ -25,7 +37,7 @@ void FText3DShapedGlyphText::CalculateWidth()
     NewLines.Reserve(Lines.Num());
 
     int32 CharacterIndex = 0;
-    for (TConstEnumerateRef<FShapedGlyphLine> GlyphLine : EnumerateRange(Lines))
+    for (const TConstEnumerateRef<FShapedGlyphLine> GlyphLine : EnumerateRange(Lines))
     {
         FShapedGlyphLine* CurrentLine = &NewLines.AddDefaulted_GetRef();
         CurrentLine->TextDirection = GlyphLine->TextDirection;
@@ -36,12 +48,15 @@ void FText3DShapedGlyphText::CalculateWidth()
         TArray<FShapedGlyphEntry> CurrentWord;
         float LineWidth = 0.0f;
         float CurrentWordLength = 0.0f;
+    	float MaxFontHeight = 0.0f;
+    	float MaxFontAscender = 0.0f;
+    	float MaxFontDescender = 0.0f;
 
         const int32 GlyphCount = GlyphLine->GlyphsToRender.Num();
 
 		for (TConstEnumerateRef<FShapedGlyphEntry> CurrentGlyph : EnumerateRange(GlyphLine->GlyphsToRender))
         {
-    		const bool bVisibleGlyph = CurrentGlyph->HasValidGlyph() && CurrentGlyph->bIsVisible;
+    		const bool bVisibleGlyph = FText3DShapedGlyphText::IsValidGlyph(*CurrentGlyph);
 			const int32 Index = bIsRTL ? Kernings.Num() - 1 - CharacterIndex : CharacterIndex;
     		const float CurrentKerning = bVisibleGlyph ? Kernings[Index] : 0.f;
     		const float GlyphAdv = GlyphLine->GetWidthAdvance(CurrentGlyph.GetIndex(), Tracking + CurrentKerning, WordSpacing);
@@ -55,6 +70,13 @@ void FText3DShapedGlyphText::CalculateWidth()
 				&& CurrentWordLength != LineWidth)  // and the line is not just a single word that we can't break
     		{
     			CurrentLine->Width = LineWidth - CurrentWordLength;
+    			CurrentLine->MaxFontHeight = MaxFontHeight;
+    			CurrentLine->MaxFontAscender = MaxFontAscender;
+    			CurrentLine->MaxFontDescender = MaxFontDescender;
+    			MaxFontHeight = 0.0f;
+    			MaxFontAscender = 0.0f;
+    			MaxFontDescender = 0.0f;
+
     			CurrentLine = &NewLines.AddDefaulted_GetRef();
     			CurrentLine->TextDirection = GlyphLine->TextDirection;
     			LineWidth = CurrentWordLength;
@@ -63,6 +85,16 @@ void FText3DShapedGlyphText::CalculateWidth()
 			LineWidth += GlyphAdv;
 			CurrentWordLength += GlyphAdv;
 
+			if (bVisibleGlyph)
+			{
+#if WITH_FREETYPE
+				const FT_Face Face = FontFaces[Index]->GetFace();
+				MaxFontHeight = FMath::Max(MaxFontHeight, Face->size->metrics.height * UE::Text3D::Metrics::FontSizeInverse);
+				MaxFontAscender = FMath::Max(MaxFontAscender, Face->size->metrics.ascender * UE::Text3D::Metrics::FontSizeInverse);
+				MaxFontDescender = FMath::Max(MaxFontDescender, Face->size->metrics.descender * UE::Text3D::Metrics::FontSizeInverse);
+#endif
+			}
+
 			if (bIsRTL)
 			{
 				CurrentWord.Insert(*CurrentGlyph, 0);
@@ -99,6 +131,9 @@ void FText3DShapedGlyphText::CalculateWidth()
         }
 
         CurrentLine->Width = LineWidth;
+    	CurrentLine->MaxFontHeight = MaxFontHeight;
+    	CurrentLine->MaxFontAscender = MaxFontAscender;
+    	CurrentLine->MaxFontDescender = MaxFontDescender;
     }
 
     Lines = NewLines;
diff --git a/Engine/Plugins/Experimental/Text3D/Source/Text3D/Private/LayoutBuilders/Text3DShapedGlyphText.h b/Engine/Plugins/Experimental/Text3D/Source/Text3D/Private/LayoutBuilders/Text3DShapedGlyphText.h
--- a/Engine/Plugins/Experimental/Text3D/Source/Text3D/Private/LayoutBuilders/Text3DShapedGlyphText.h
+++ b/Engine/Plugins/Experimental/Text3D/Source/Text3D/Private/LayoutBuilders/Text3DShapedGlyphText.h
@@ -8,10 +8,13 @@
 
 class FFreeTypeFace;
 struct FShapedGlyphLine;
+struct FT_FaceRec_;
 
 /** Contains text metrics with sufficient information to fetch and transform each line. */
 struct FText3DShapedGlyphText
 {
+	static bool IsValidGlyph(const FShapedGlyphEntry& InGlyph);
+
 	FText3DShapedGlyphText();
 
 	void Reset();
@@ -20,6 +23,9 @@ struct FText3DShapedGlyphText
 	/** Individual character kerning */
 	TArray<float> Kernings;
 
+	/** Individual character font faces */
+	TArray<TSharedPtr<FFreeTypeFace>> FontFaces;
+
 	/** General tracking for all characters */
 	float Tracking;
 
@@ -34,7 +40,4 @@ struct FText3DShapedGlyphText
 
 	/** Lines calculated based on current text */
 	TArray<FShapedGlyphLine> Lines;
-
-	/** Glyph index to font face */
-	TMap<uint32, const FFreeTypeFace*> GlyphIndexToFontFace;
 };
diff --git a/Engine/Plugins/Experimental/Text3D/Source/Text3D/Private/Renderers/DynamicMesh/Text3DDynamicMeshRenderer.cpp b/Engine/Plugins/Experimental/Text3D/Source/Text3D/Private/Renderers/DynamicMesh/Text3DDynamicMeshRenderer.cpp
--- a/Engine/Plugins/Experimental/Text3D/Source/Text3D/Private/Renderers/DynamicMesh/Text3DDynamicMeshRenderer.cpp
+++ b/Engine/Plugins/Experimental/Text3D/Source/Text3D/Private/Renderers/DynamicMesh/Text3DDynamicMeshRenderer.cpp
@@ -84,12 +84,13 @@ void UText3DDynamicMeshRenderer::OnUpdate(const UE::Text3D::Renderer::FUpdatePar
 	if (InParameters.CurrentFlag == EText3DRendererFlags::Geometry)
 	{
 		AllocateGlyphMeshData(MeshEditor, TextComponent->GetCharacterCount());
-
-		const UText3DGeometryExtensionBase* GeometryExtension = TextComponent->GetGeometryExtension();
-		TextComponent->ForEachCharacter([this, GeometryExtension, &MeshEditor](UText3DCharacterBase* InCharacter, uint16 InIndex, uint16 InCount)
+		
+		TextComponent->ForEachCharacter([this, &MeshEditor](UText3DCharacterBase* InCharacter, uint16 InIndex, uint16 InCount)
 		{
-			const FText3DCachedMesh* CachedGlyph = GeometryExtension->FindOrLoadGlyphMesh(InCharacter);
-			AppendGlyphMesh(MeshEditor, InIndex, InCharacter->GetGlyphIndex(), CachedGlyph->DynamicMesh);
+			if (const FText3DCachedMesh* CachedGlyph = InCharacter->GetGlyphMesh())
+			{
+				AppendGlyphMesh(MeshEditor, InIndex, InCharacter->GetGlyphIndex(), CachedGlyph->DynamicMesh);
+			}
 		});
 
 		RefreshBounds();
diff --git a/Engine/Plugins/Experimental/Text3D/Source/Text3D/Private/Renderers/StaticMeshes/Text3DStaticMeshesRenderer.cpp b/Engine/Plugins/Experimental/Text3D/Source/Text3D/Private/Renderers/StaticMeshes/Text3DStaticMeshesRenderer.cpp
--- a/Engine/Plugins/Experimental/Text3D/Source/Text3D/Private/Renderers/StaticMeshes/Text3DStaticMeshesRenderer.cpp
+++ b/Engine/Plugins/Experimental/Text3D/Source/Text3D/Private/Renderers/StaticMeshes/Text3DStaticMeshesRenderer.cpp
@@ -2,7 +2,6 @@
 
 #include "Renderers/StaticMeshes/Text3DStaticMeshesRenderer.h"
 
-#include "Algo/Accumulate.h"
 #include "Characters/Text3DCharacterBase.h"
 #include "Components/SceneComponent.h"
 #include "Components/StaticMeshComponent.h"
@@ -46,14 +45,18 @@ void UText3DStaticMeshesRenderer::OnUpdate(const UE::Text3D::Renderer::FUpdatePa
 	{
 		AllocateComponents(TextComponent->GetCharacterCount());
 
-		const UText3DGeometryExtensionBase* GeometryExtension = TextComponent->GetGeometryExtension();
-
-		TextComponent->ForEachCharacter([this, GeometryExtension](UText3DCharacterBase* InCharacter, uint16 InIndex, uint16 InCount)
+		TextComponent->ForEachCharacter([this](UText3DCharacterBase* InCharacter, uint16 InIndex, uint16 InCount)
 		{
-			const FText3DCachedMesh* CachedMesh = GeometryExtension->FindOrLoadGlyphMesh(InCharacter);
 			if (CharacterMeshes.IsValidIndex(InIndex))
 			{
-				CharacterMeshes[InIndex]->SetStaticMesh(CachedMesh->StaticMesh);
+				if (const FText3DCachedMesh* CachedMesh = InCharacter->GetGlyphMesh())
+				{
+					CharacterMeshes[InIndex]->SetStaticMesh(CachedMesh->StaticMesh);
+				}
+				else
+				{
+					CharacterMeshes[InIndex]->SetStaticMesh(nullptr);
+				}
 			}
 			else
 			{
@@ -152,11 +155,14 @@ void UText3DStaticMeshesRenderer::OnClear()
 
 void UText3DStaticMeshesRenderer::OnDestroy()
 {
+	FDetachmentTransformRules DetachmentRules = FDetachmentTransformRules::KeepRelativeTransform;
+	DetachmentRules.bCallModify = false;
+
 	for (UStaticMeshComponent* MeshComponent : CharacterMeshes)
 	{
 		if (MeshComponent)
 		{
-			MeshComponent->DetachFromComponent(FDetachmentTransformRules::KeepRelativeTransform);
+			MeshComponent->DetachFromComponent(DetachmentRules);
 			MeshComponent->DestroyComponent();
 		}
 	}
@@ -208,13 +214,17 @@ void UText3DStaticMeshesRenderer::AllocateComponents(int32 InCount)
 	const int32 RemainingCharacterCount = CharacterMeshes.Num() - InCount;
 
 	AActor* Actor = GetTypedOuter<AActor>();
+
 	if (RemainingCharacterCount > 0)
 	{
+		FDetachmentTransformRules DetachmentRule = FDetachmentTransformRules::KeepRelativeTransform;
+		DetachmentRule.bCallModify = false;
+
 		for (int32 CharacterIndex = InCount; CharacterIndex < CharacterMeshes.Num(); CharacterIndex++)
 		{
 			if (UStaticMeshComponent* MeshComponent = CharacterMeshes[CharacterIndex])
 			{
-				MeshComponent->DetachFromComponent(FDetachmentTransformRules::KeepRelativeTransform);
+				MeshComponent->DetachFromComponent(DetachmentRule);
 				MeshComponent->UnregisterComponent();
 				Actor->RemoveInstanceComponent(MeshComponent);
 				Actor->RemoveOwnedComponent(MeshComponent);
@@ -238,9 +248,9 @@ void UText3DStaticMeshesRenderer::AllocateComponents(int32 InCount)
 				CharacterMeshes.Add(MeshComponent);
 			}
 
+			Actor->AddOwnedComponent(MeshComponent);
 			MeshComponent->RegisterComponent();
 			MeshComponent->AttachToComponent(TextRoot, FAttachmentTransformRules::KeepRelativeTransform);
-			Actor->AddOwnedComponent(MeshComponent);
 		}
 	}
 
diff --git a/Engine/Plugins/Experimental/Text3D/Source/Text3D/Private/Renderers/Text3DRendererBase.cpp b/Engine/Plugins/Experimental/Text3D/Source/Text3D/Private/Renderers/Text3DRendererBase.cpp
--- a/Engine/Plugins/Experimental/Text3D/Source/Text3D/Private/Renderers/Text3DRendererBase.cpp
+++ b/Engine/Plugins/Experimental/Text3D/Source/Text3D/Private/Renderers/Text3DRendererBase.cpp
@@ -2,8 +2,9 @@
 
 #include "Renderers/Text3DRendererBase.h"
 
-#include "Extensions/Text3DGeometryExtensionBase.h"
 #include "Extensions/Text3DCharacterExtensionBase.h"
+#include "Extensions/Text3DFormatStyleExtensionBase.h"
+#include "Extensions/Text3DGeometryExtensionBase.h"
 #include "Extensions/Text3DLayoutEffectBase.h"
 #include "Extensions/Text3DLayoutExtensionBase.h"
 #include "Extensions/Text3DMaterialExtensionBase.h"
@@ -105,6 +106,7 @@ void UText3DRendererBase::Update(EText3DRendererFlags InFlags)
 
 	TArray<UText3DExtensionBase*, TInlineAllocator<6>> Extensions
 	{
+		Text3DComponent->GetFormatStyleExtension(),
 		Text3DComponent->GetCharacterExtension(),
 		Text3DComponent->GetGeometryExtension(),
 		Text3DComponent->GetLayoutExtension(),
diff --git a/Engine/Plugins/Experimental/Text3D/Source/Text3D/Private/Settings/Text3DProjectSettings.cpp b/Engine/Plugins/Experimental/Text3D/Source/Text3D/Private/Settings/Text3DProjectSettings.cpp
--- a/Engine/Plugins/Experimental/Text3D/Source/Text3D/Private/Settings/Text3DProjectSettings.cpp
+++ b/Engine/Plugins/Experimental/Text3D/Source/Text3D/Private/Settings/Text3DProjectSettings.cpp
@@ -6,6 +6,7 @@
 #include "Engine/FontFace.h"
 #include "Materials/Material.h"
 #include "Modules/ModuleManager.h"
+#include "Subsystems/Text3DEngineSubsystem.h"
 #include "UObject/SoftObjectPath.h"
 
 #if WITH_EDITOR
@@ -75,6 +76,38 @@ UMaterial* UText3DProjectSettings::GetBaseMaterial(const FText3DMaterialKey& InK
 	return nullptr;
 }
 
+void UText3DProjectSettings::SetScheduleFontFaceGlyphCleanup(bool bInEnabled)
+{
+	if (bScheduleFontFaceGlyphCleanup == bInEnabled)
+	{
+		return;
+	}
+
+	bScheduleFontFaceGlyphCleanup = bInEnabled;
+	OnFontFaceCleanupOptionsChanged();
+}
+
+bool UText3DProjectSettings::GetScheduleFontFaceGlyphCleanup() const
+{
+	return bScheduleFontFaceGlyphCleanup;
+}
+
+void UText3DProjectSettings::SetFontFaceGlyphCleanupPeriod(float InValue)
+{
+	if (FMath::IsNearlyEqual(FontFaceGlyphCleanupPeriod, InValue))
+	{
+		return;
+	}
+
+	FontFaceGlyphCleanupPeriod = InValue;
+	OnFontFaceCleanupOptionsChanged();
+}
+
+float UText3DProjectSettings::GetFontFaceGlyphCleanupPeriod() const
+{
+	return FontFaceGlyphCleanupPeriod;
+}
+
 UText3DProjectSettings::UText3DProjectSettings()
 {
 	// Config
@@ -114,6 +147,14 @@ void UText3DProjectSettings::AddMaterial(FText3DMaterialKey&& InKey, TSoftObject
 	BaseMaterials.Emplace(InKey, InMaterial);
 }
 
+void UText3DProjectSettings::OnFontFaceCleanupOptionsChanged()
+{
+	if (UText3DEngineSubsystem* TextSubsystem = UText3DEngineSubsystem::Get())
+	{
+		TextSubsystem->ScheduleCacheCleanup();
+	}
+}
+
 #if WITH_EDITOR
 void UText3DProjectSettings::PostEditChangeProperty(struct FPropertyChangedEvent& InChangeEvent)
 {
@@ -132,12 +173,23 @@ void UText3DProjectSettings::PostEditChangeProperty(struct FPropertyChangedEvent
 		GET_MEMBER_NAME_CHECKED(UText3DProjectSettings, bDebugMode),
 	};
 
+	static const TSet<FName> FontFaceCleanupPropertyNames
+	{
+		GET_MEMBER_NAME_CHECKED(UText3DProjectSettings, bScheduleFontFaceGlyphCleanup),
+		GET_MEMBER_NAME_CHECKED(UText3DProjectSettings, FontFaceGlyphCleanupPeriod)
+	};
+
 	const FName MemberName = InChangeEvent.GetMemberPropertyName();
 
-	if (PropertyNames.Contains(MemberName))
+	if (PropertyNames.Contains(MemberName) || FontFaceCleanupPropertyNames.Contains(MemberName))
 	{
 		Save();
 	}
+
+	if (FontFaceCleanupPropertyNames.Contains(MemberName))
+	{
+		OnFontFaceCleanupOptionsChanged();
+	}
 }
 
 TConstArrayView<FString> UText3DProjectSettings::GetFavoriteFonts() const
diff --git a/Engine/Plugins/Experimental/Text3D/Source/Text3D/Private/Styles/Text3DFormatStyleBase.cpp b/Engine/Plugins/Experimental/Text3D/Source/Text3D/Private/Styles/Text3DFormatStyleBase.cpp
new file mode 100644
--- /dev/null
+++ b/Engine/Plugins/Experimental/Text3D/Source/Text3D/Private/Styles/Text3DFormatStyleBase.cpp
@@ -0,0 +1,144 @@
+// Copyright Epic Games, Inc. All Rights Reserved.
+
+#include "Styles/Text3DFormatStyleBase.h"
+
+#include "Engine/Font.h"
+#include "Text3DComponent.h"
+#include "Settings/Text3DProjectSettings.h"
+#include "Styling/StyleDefaults.h"
+
+#if WITH_EDITOR
+void UText3DFormatStyleBase::PostEditChangeProperty(FPropertyChangedEvent& InPropertyChangedEvent)
+{
+	Super::PostEditChangeProperty(InPropertyChangedEvent);
+
+	const FName MemberProperty = InPropertyChangedEvent.GetMemberPropertyName();
+
+	if (MemberProperty == GET_MEMBER_NAME_CHECKED(UText3DFormatStyleBase, Font))
+	{
+		OnFontChanged();
+	}
+	else if (MemberProperty == GET_MEMBER_NAME_CHECKED(UText3DFormatStyleBase, FontTypeface)
+		|| MemberProperty == GET_MEMBER_NAME_CHECKED(UText3DFormatStyleBase, StyleName))
+	{
+		OnStylePropertiesChanged();
+	}
+}
+#endif
+
+UText3DFormatStyleBase::UText3DFormatStyleBase()
+{
+	if (const UText3DProjectSettings* Text3DSettings = UText3DProjectSettings::Get())
+	{
+		Font = Text3DSettings->GetFallbackFont();
+		OnFontChanged();
+	}
+}
+
+const FTextBlockStyle& UText3DFormatStyleBase::GetTextStyle() const
+{
+	return TextStyle;
+}
+
+void UText3DFormatStyleBase::SetStyleName(FName InName)
+{
+	if (InName == StyleName)
+	{
+		return;
+	}
+
+	StyleName = InName;
+	OnStylePropertiesChanged();
+}
+
+void UText3DFormatStyleBase::SetFont(UFont* InFont)
+{
+	if (InFont == Font)
+	{
+		return;
+	}
+
+	Font = InFont;
+	OnFontChanged();
+}
+
+void UText3DFormatStyleBase::SetFontTypeface(FName InTypeface)
+{
+	if (InTypeface == FontTypeface)
+	{
+		return;
+	}
+
+	if (!GetTypefaceNames().Contains(InTypeface))
+	{
+		return;
+	}
+
+	FontTypeface = InTypeface;
+	OnStylePropertiesChanged();
+}
+
+void UText3DFormatStyleBase::PostLoad()
+{
+	Super::PostLoad();
+	UpdateTextStyle();
+}
+
+void UText3DFormatStyleBase::OnStylePropertiesChanged()
+{
+	UpdateTextStyle();
+	if (!StyleName.IsNone())
+	{
+		if (UText3DComponent* Text3DComponent = GetTypedOuter<UText3DComponent>())
+		{
+			Text3DComponent->RequestUpdate(EText3DRendererFlags::All, /** Immediate */false);
+		}
+	}
+}
+
+void UText3DFormatStyleBase::UpdateTextStyle()
+{
+	FSlateFontInfo FontInfo(Font, UE::Text3D::Metrics::SlateFontSize);
+	FontInfo.CompositeFont = Font && Font->GetCompositeFont() ? MakeShared<FCompositeFont>(*Font->GetCompositeFont()) : FStyleDefaults::GetFontInfo().CompositeFont;
+	FontInfo.TypefaceFontName = FontTypeface;
+	FontInfo.FontFallback = EFontFallback::FF_NoFallback;
+	TextStyle.SetFont(FontInfo);
+}
+
+void UText3DFormatStyleBase::OnFontChanged()
+{
+	const TArray<FName> Typefaces = GetTypefaceNames();
+	if (!Typefaces.Contains(FontTypeface) && !Typefaces.IsEmpty())
+	{
+		FontTypeface = Typefaces[0];
+	}
+	OnStylePropertiesChanged();
+}
+
+TArray<FName> UText3DFormatStyleBase::GetTypefaceNames() const
+{
+	TArray<FName> TypefaceNames;
+
+	if (Font)
+	{
+		for (const FTypefaceEntry& TypeFaceFont : GetAvailableTypefaces())
+		{
+			TypefaceNames.Add(TypeFaceFont.Name);
+		}
+	}
+
+	return TypefaceNames;
+}
+
+TArrayView<const FTypefaceEntry> UText3DFormatStyleBase::GetAvailableTypefaces() const
+{
+	if (Font)
+	{
+		if (const FCompositeFont* CompositeFont = Font->GetCompositeFont())
+		{
+			return CompositeFont->DefaultTypeface.Fonts;
+		}
+	}
+
+	return {};
+}
diff --git a/Engine/Plugins/Experimental/Text3D/Source/Text3D/Private/Subsystems/Text3DEngineSubsystem.cpp b/Engine/Plugins/Experimental/Text3D/Source/Text3D/Private/Subsystems/Text3DEngineSubsystem.cpp
--- a/Engine/Plugins/Experimental/Text3D/Source/Text3D/Private/Subsystems/Text3DEngineSubsystem.cpp
+++ b/Engine/Plugins/Experimental/Text3D/Source/Text3D/Private/Subsystems/Text3DEngineSubsystem.cpp
@@ -4,7 +4,6 @@
 
 #include "Containers/Ticker.h"
 #include "UDynamicMesh.h"
-#include "Engine/Font.h"
 #include "Engine/Engine.h"
 #include "Engine/StaticMesh.h"
 #include "GeometryBuilders/Text3DGlyphContourNode.h"
@@ -13,27 +12,13 @@
 #include "Logs/Text3DLogs.h"
 #include "Materials/Material.h"
 #include "Misc/FileHelper.h"
+#include "Settings/Text3DProjectSettings.h"
 #include "UObject/ConstructorHelpers.h"
 
 #if WITH_FREETYPE
 #include "Fonts/FontCacheFreeType.h"
 #endif
 
-FCachedFontMeshes::FCachedFontMeshes()
-{
-	CacheCounter = MakeShared<int32>(0);
-}
-
-int32 FCachedFontMeshes::GetCacheCount() const
-{
-	return CacheCounter.GetSharedReferenceCount();
-}
-
-TSharedPtr<int32> FCachedFontMeshes::GetCacheCounter()
-{
-	return CacheCounter;
-}
-
 UText3DEngineSubsystem* UText3DEngineSubsystem::Get()
 {
 	if (GEngine)
@@ -44,39 +29,16 @@ UText3DEngineSubsystem* UText3DEngineSubsystem::Get()
 	return nullptr;
 }
 
-UText3DEngineSubsystem::UText3DEngineSubsystem()
+void UText3DEngineSubsystem::Initialize(FSubsystemCollectionBase& InCollection)
 {
-	if (!IsRunningDedicatedServer())
-	{
-		struct FConstructorStatics
-		{
-			ConstructorHelpers::FObjectFinder<UMaterial> Material;
-			FConstructorStatics() :
-				Material(TEXT("/Engine/BasicShapes/BasicShapeMaterial"))
-			{
-			}
-		};
-		static FConstructorStatics ConstructorStatics;
-
-		DefaultMaterial = ConstructorStatics.Material.Object;
-	}
-}
-
-void UText3DEngineSubsystem::Initialize(FSubsystemCollectionBase& Collection)
-{
-	Super::Initialize(Collection);
-	CleanupTickerHandle = FTSTicker::GetCoreTicker().AddTicker(FTickerDelegate::CreateUObject(this, &UText3DEngineSubsystem::CleanupTimerCallback), 600.0f);
+	Super::Initialize(InCollection);
+	ScheduleCacheCleanup();
 }
 
 void UText3DEngineSubsystem::Deinitialize()
 {
-	if (CleanupTickerHandle.IsValid())
-	{
-		FTSTicker::GetCoreTicker().RemoveTicker(CleanupTickerHandle);
-		CleanupTickerHandle.Reset();
-	}
-
 	Super::Deinitialize();
+	ClearCache();
 }
 
 bool UText3DEngineSubsystem::Exec_Dev(UWorld* InWorld, const TCHAR* InCmd, FOutputDevice& InAr)
@@ -95,11 +57,10 @@ bool UText3DEngineSubsystem::Exec_Dev(UWorld* InWorld, const TCHAR* InCmd, FOutp
 				PrintCache();
 				return true;
 			}
-			else if (FParse::Command(&InCmd, TEXT("clear")))
-			{
-				Cleanup();
-				PrintCache();
 
+			if (FParse::Command(&InCmd, TEXT("cleanup")))
+			{
+				CleanupCache(0);
 				return true;
 			}
 		}
@@ -108,218 +69,121 @@ bool UText3DEngineSubsystem::Exec_Dev(UWorld* InWorld, const TCHAR* InCmd, FOutp
 	return false;
 }
 
-void UText3DEngineSubsystem::Reset()
+void UText3DEngineSubsystem::ScheduleCacheCleanup()
 {
-	CachedFonts.Reset();
-}
-
-bool UText3DEngineSubsystem::CleanupTimerCallback(float DeltaTime)
-{
-	Cleanup();
-	return true;
-}
+	CancelCacheCleanup();
 
-void UText3DEngineSubsystem::PrintCache() const
-{
-	UE_LOG(LogText3D, Log, TEXT("Text3D Engine Subsystem is currently caching the following fonts:"));
-	for (const TPair<uint32, FCachedFontData>& CachedFont : CachedFonts)
+	if (const UText3DProjectSettings* TextSettings = UText3DProjectSettings::Get())
 	{
-		CachedFont.Value.PrintCache();
+		if (TextSettings->GetScheduleFontFaceGlyphCleanup())
+		{
+			CacheCleanupDelegate = FTSTicker::GetCoreTicker().AddTicker(
+				FTickerDelegate::CreateUObject(this, &UText3DEngineSubsystem::CleanupCache),
+				FMath::Max(0, TextSettings->GetFontFaceGlyphCleanupPeriod())
+			);
+		}
 	}
 }
 
-void UText3DEngineSubsystem::Cleanup()
+void UText3DEngineSubsystem::CancelCacheCleanup()
 {
-	QUICK_SCOPE_CYCLE_COUNTER(STAT_UText3DEngineSubsystem_Cleanup);
-
-	UE_LOG(LogText3D, Log, TEXT("Text3D Engine Subsystem is clearing unreferenced cached fonts and typefaces."));
-	for (auto It = CachedFonts.CreateIterator(); It; ++It)
+	if (CacheCleanupDelegate.IsValid())
 	{
-		if (It.Value().Cleanup())
-		{
-			It.RemoveCurrent();
-		}
+		FTSTicker::GetCoreTicker().RemoveTicker(CacheCleanupDelegate);
+		CacheCleanupDelegate.Reset();
 	}
 }
 
-FCachedFontData& UText3DEngineSubsystem::GetCachedFontData(UFont* InFont, int32 InTypefaceEntryIndex)
+bool UText3DEngineSubsystem::CleanupCache(float)
 {
-	const uint32 FontHash = HashCombine(0, GetTypeHash(InFont));
+	uint32 GlyphMeshCount = 0;
+	uint32 GlyphMeshRemoved = 0;
+	uint32 FontFaceRemoved = 0;
 
-	if (CachedFonts.Contains(FontHash))
+	for (TMap<uint32, FText3DFontFaceCache>::TIterator It(CachedFontFaces); It; ++It)
 	{
-		uint32 TypefaceFontDataHash = 0;
-
-		// First we check if the font itself has changed since we last cached it
-		const FCompositeFont* CompositeFont = InFont->GetCompositeFont();
-		if (CompositeFont && !CompositeFont->DefaultTypeface.Fonts.IsEmpty())
-		{
-			FTypefaceEntry Typeface;
-			if (CompositeFont->DefaultTypeface.Fonts.IsValidIndex(InTypefaceEntryIndex))
-			{
-				Typeface = CompositeFont->DefaultTypeface.Fonts[InTypefaceEntryIndex];
-			}
-			else
-			{
-				Typeface = CompositeFont->DefaultTypeface.Fonts[0];
-			}
-
-			TypefaceFontDataHash = HashCombine(0, GetTypeHash(Typeface.Font));
-		}
-
-		const uint32 CachedTypefaceFontDataHash = CachedFonts[FontHash].GetTypefaceFontDataHash(InTypefaceEntryIndex);
-		if (CachedTypefaceFontDataHash != TypefaceFontDataHash)
+		FText3DFontFaceCache& FontFace = It.Value();
+		GlyphMeshRemoved += FontFace.CleanupCache();
+		GlyphMeshCount += FontFace.GetCacheGlyphCount();
+		if (!FontFace.IsValid() || FontFace.GetCacheGlyphCount() == 0)
 		{
-			// FontFace has changed, remove the cached data
-			CachedFonts[FontHash].CleanupTypeface(InTypefaceEntryIndex);
+			It.RemoveCurrent();
+			FontFaceRemoved++;
 		}
 	}
 
-	if (!CachedFonts.Contains(FontHash))
+	if (FontFaceRemoved > 0 || GlyphMeshRemoved > 0)
 	{
-		FCachedFontData& NewCachedFontData = CachedFonts.Add(FontHash);
-		NewCachedFontData.SetFont(InFont);
+		UE_LOG(LogText3D, Log, TEXT("Text3D Engine Subsystem cache cleanup completed : %i unused glyph(s) removed, %i remaining glyph(s), %i unused font face(s) removed, %i remaining font face(s)"), GlyphMeshRemoved, GlyphMeshCount, FontFaceRemoved, CachedFontFaces.Num());
 	}
 
-	FCachedFontData& CachedFontData = CachedFonts[FontHash];
-	CachedFontData.LoadFreeTypeFace(InTypefaceEntryIndex);
-
-	return CachedFonts[FontHash];
-}
-
-FTypefaceFontData::FTypefaceFontData()
-{
-	CacheCounter = MakeShared<int32>(0);
-}
-
-#if WITH_FREETYPE
-void FTypefaceFontData::SetTypeFace(FT_Face InTypeFace)
-{
-	TypeFace = InTypeFace;
-}
-
-void FTypefaceFontData::SetTypeFaceData(const TSharedPtr<const FFontFaceData>& InTypeFaceData)
-{
-	TypeFaceDataWeak = InTypeFaceData;
-}
-
-TSharedPtr<const FFontFaceData> FTypefaceFontData::GetTypeFaceData() const
-{
-	return TypeFaceDataWeak.Pin();
+	return true;
 }
-#endif
 
-FCachedFontData::FCachedFontData()
+void UText3DEngineSubsystem::PrintCache() const
 {
-	Font = nullptr;
+	UE_LOG(LogText3D, Log, TEXT("Text3D Engine Subsystem is currently caching %i fonts: "), CachedFontFaces.Num());
+	for (const TPair<uint32, FText3DFontFaceCache>& CachedFont : CachedFontFaces)
+	{
+		CachedFont.Value.PrintCache();
+	}
 }
 
-FCachedFontData::~FCachedFontData()
+void UText3DEngineSubsystem::ClearCache()
 {
-	ClearFreeTypeFace();
+	CachedFontFaces.Empty();
 }
 
-void FCachedFontData::ClearFreeTypeFace()
+FText3DFontFaceCache* UText3DEngineSubsystem::FindCachedFontFace(uint32 InFontFaceHash)
 {
-	TypefaceFontDataMap.Reset();
+	return CachedFontFaces.Find(InFontFaceHash);
 }
 
-void FCachedFontData::LoadFreeTypeFace(uint32 InTypefaceIndex)
+#if WITH_FREETYPE
+FText3DFontFaceCache* UText3DEngineSubsystem::FindOrAddCachedFontFace(const TSharedPtr<FFreeTypeFace>& InFontFace)
 {
-	if (!Font)
+	if (!InFontFace.IsValid() || !InFontFace->IsFaceValid())
 	{
-		return;
-	}
-
-	const FCompositeFont* const CompositeFont = Font->GetCompositeFont();
-	if (!CompositeFont || CompositeFont->DefaultTypeface.Fonts.IsEmpty())
-	{
-		return;
-	}
-
-	if (!CompositeFont->DefaultTypeface.Fonts.IsValidIndex(InTypefaceIndex))
-	{
-		return;
+		return nullptr;
 	}
 
-	const FTypefaceEntry& Typeface = CompositeFont->DefaultTypeface.Fonts[InTypefaceIndex];
-	const FFontFaceDataConstPtr FaceData = Typeface.Font.GetFontFaceData();
-
-	FTypefaceFontData& TypefaceFontData = TypefaceFontDataMap.FindOrAdd(InTypefaceIndex);
-	TypefaceFontData.SetTypefaceName(Typeface.Name);
-	TypefaceFontData.SetTypefaceFontDataHash(HashCombine(0, GetTypeHash(Typeface.Font)));
-#if WITH_FREETYPE
-	TypefaceFontData.SetTypeFaceData(FaceData);
-#endif
+	const uint32 FontFaceHash = FText3DFontFaceCache::GetFontFaceHash(InFontFace);
+	return FontFaceHash != 0 ? &CachedFontFaces.FindOrAdd(FontFaceHash, FText3DFontFaceCache(InFontFace)) : nullptr;
 }
 
-uint32 FCachedFontData::GetTypefaceFontDataHash(int32 InTypefaceEntryIndex) const
+uint32 FText3DFontFaceCache::GetFontFaceHash(const TSharedPtr<FFreeTypeFace>& InFontFace)
 {
-	if (TypefaceFontDataMap.Contains(InTypefaceEntryIndex))
+	if (const FT_Face Face = InFontFace->GetFace())
 	{
-		return TypefaceFontDataMap[InTypefaceEntryIndex].GetTypefaceFontDataHash();
+		const FString FaceFamily(Face->family_name);
+		const FString FaceStyle(Face->style_name);
+		return HashCombine(GetTypeHash(FaceFamily), GetTypeHash(FaceStyle));
 	}
 
 	return 0;
 }
 
-const FTypefaceFontData* FCachedFontData::GetTypefaceFontData(int32 InTypefaceEntryIndex) const
-{
-	return TypefaceFontDataMap.Find(InTypefaceEntryIndex);
-}
-
-FTypefaceFontData* FCachedFontData::GetTypefaceFontData(int32 InTypefaceEntryIndex)
+FText3DFontFaceCache::FText3DFontFaceCache(const TSharedPtr<FFreeTypeFace>& InFontFace)
+	: FontFace(InFontFace)
+	, FontFaceHash(GetFontFaceHash(InFontFace))
 {
-	return TypefaceFontDataMap.Find(InTypefaceEntryIndex);
-}
-
-TSharedPtr<int32> FCachedFontData::GetCacheCounter(int32 InTypefaceEntryIndex)
-{
-	if (TypefaceFontDataMap.Contains(InTypefaceEntryIndex))
-	{
-		return TypefaceFontDataMap[InTypefaceEntryIndex].GetCacheCounter();
-	}
-
-	return nullptr;
 }
+#endif
 
-TSharedPtr<int32> FCachedFontData::GetMeshesCacheCounter(const FGlyphMeshParameters& InParameters)
+uint32 FText3DFontFaceCache::GetGlyphMeshHash(uint32 InGlyphIndex, const FGlyphMeshParameters& InParameters)
 {
-	const uint32 HashParameters = GetTypeHash(InParameters);
-
-	const uint32 TypefaceIndex = InParameters.TypefaceIndex;
-	if (TypefaceFontDataMap.Contains(TypefaceIndex))
-	{
-		FCachedFontMeshes& CachedMeshes = TypefaceFontDataMap[TypefaceIndex].FindOrAddMeshes(HashParameters);
-		return CachedMeshes.GetCacheCounter();
-	}
-
-	return nullptr;
+	return HashCombine(GetTypeHash(InParameters), GetTypeHash(InGlyphIndex));
 }
 
-FText3DCachedMesh* FCachedFontData::GetGlyphMesh(uint32 InGlyphIndex, const FGlyphMeshParameters& InParameters)
+FText3DCachedMesh* FText3DFontFaceCache::FindOrBuildGlyphMesh(uint32 InGlyphIndex, const FGlyphMeshParameters& InParameters)
 {
-	const uint32 TypefaceIndex = InParameters.TypefaceIndex;
-	const uint32 HashParameters = GetTypeHash(InParameters);
-
-	if (!TypefaceFontDataMap.Contains(TypefaceIndex))
-	{
-		return nullptr;
-	}
-
-	FCachedFontMeshes& CachedMeshes = TypefaceFontDataMap[TypefaceIndex].FindOrAddMeshes(HashParameters);
-
-	if (FText3DCachedMesh* CachedMesh = CachedMeshes.Glyphs.Find(InGlyphIndex))
+	const uint32 HashParameters = GetGlyphMeshHash(InGlyphIndex, InParameters);
+	if (FText3DCachedMesh* CachedMesh = GlyphMeshes.Find(HashParameters))
 	{
 		return CachedMesh;
 	}
 
-	uint32 HashGroup = 0;
-	HashGroup = HashCombine(HashGroup, GetTypeHash(Font));
-	HashGroup = HashCombine(HashGroup, GetTypeHash(InGlyphIndex));
-	const FString MeshName = FString::Printf(TEXT("Text3D_Char_%u_%u"), HashGroup, HashParameters);
-
-	const TText3DGlyphContourNodeShared Root = GetGlyphContours(InGlyphIndex, TypefaceIndex);
+	const TText3DGlyphContourNodeShared Root = LoadGlyphContours(InGlyphIndex);
 	if (!Root || Root->Children.IsEmpty())
 	{
 		return nullptr;
@@ -331,8 +195,9 @@ FText3DCachedMesh* FCachedFontData::GetGlyphMesh(uint32 InGlyphIndex, const FGly
 		return nullptr;
 	}
 
-	FText3DCachedMesh& CachedMesh = CachedMeshes.Glyphs.FindOrAdd(InGlyphIndex);
+	FText3DCachedMesh& CachedMesh = GlyphMeshes.FindOrAdd(HashParameters);
 
+	const FString MeshName = FString::Printf(TEXT("Text3D_Char_%u"), HashParameters);
 	const FName StaticMeshName = MakeUniqueObjectName(TextSubsystem, UStaticMesh::StaticClass(), FName(*MeshName));
 	CachedMesh.StaticMesh = NewObject<UStaticMesh>(TextSubsystem, StaticMeshName, RF_Transient | RF_Public);
 
@@ -353,48 +218,68 @@ FText3DCachedMesh* FCachedFontData::GetGlyphMesh(uint32 InGlyphIndex, const FGly
 	MeshCreator.SetFrontAndBevelTextureCoordinates(InParameters.Bevel);
 	MeshCreator.MirrorGroups(InParameters.Extrude);
 	MeshCreator.MovePivot(InParameters.PivotOffset);
-	MeshCreator.BuildMesh(CachedMesh, TextSubsystem->DefaultMaterial);
+	MeshCreator.BuildMesh(CachedMesh, GetDefault<UText3DProjectSettings>()->GetDefaultMaterial());
 
 	return &CachedMesh;
 }
 
-#if WITH_FREETYPE
-FT_Face FCachedFontData::GetFreeTypeFace(uint32 InTypefaceIndex)
+const FText3DCachedMesh* FText3DFontFaceCache::FindGlyphMesh(uint32 InGlyphHash) const
 {
-	if (const FTypefaceFontData* TypeFaceFontData = TypefaceFontDataMap.Find(InTypefaceIndex))
+	if (const FText3DCachedMesh* CachedMesh = GlyphMeshes.Find(InGlyphHash))
 	{
-		return TypeFaceFontData->GetTypeFace();
+		return CachedMesh;
 	}
 
 	return nullptr;
 }
-#endif
 
-FString FCachedFontData::GetFontName() const
+FText3DCachedMesh* FText3DFontFaceCache::FindGlyphMesh(uint32 InGlyphHash)
 {
-	if (Font)
+	if (FText3DCachedMesh* CachedMesh = GlyphMeshes.Find(InGlyphHash))
 	{
-		return Font->GetName();
+		return CachedMesh;
 	}
 
-	return TEXT("");
+	return nullptr;
+}
+
+uint32 FText3DFontFaceCache::CleanupCache()
+{
+	uint32 GlyphMeshRemoved = 0;
+	for (TMap<uint32, FText3DCachedMesh>::TIterator It(GlyphMeshes); It; ++It)
+	{
+		if (It->Value.RefCount == 0)
+		{
+			It.RemoveCurrent();
+			GlyphMeshRemoved++;
+		}
+	}
+	return GlyphMeshRemoved;
 }
 
-TText3DGlyphContourNodeShared FCachedFontData::GetGlyphContours(uint32 InGlyphIndex, int32 InTypefaceEntryIndex)
+TText3DGlyphContourNodeShared FText3DFontFaceCache::LoadGlyphContours(uint32 InGlyphIndex) const
 {
 #if WITH_FREETYPE
-	FT_Face Face = GetFreeTypeFace(InTypefaceEntryIndex);
+	if (!FontFace)
+	{
+		return nullptr;
+	}
+
+	const FT_Face Face = FontFace->GetFace();
 	if (!Face)
 	{
+		UE_LOG(LogText3D, Error, TEXT("Failed to load font face glyph contours '%u %i' due to invalid face"), FontFaceHash, InGlyphIndex);
 		return nullptr;
 	}
 
-	if (FT_Load_Glyph(Face, InGlyphIndex, FT_LOAD_NO_HINTING | FT_LOAD_NO_BITMAP) != 0)
+	const FT_Error Result = FT_Load_Glyph(Face, InGlyphIndex, FT_LOAD_DEFAULT);
+	if (Result != 0)
 	{
+		UE_LOG(LogText3D, Error, TEXT("Failed to load font face glyph contours '%u %i' with error code : %i"), FontFaceHash, InGlyphIndex, Result);
 		return nullptr;
 	}
 
-	FText3DGlyphLoader GlyphLoader(Face->glyph);
+	const FText3DGlyphLoader GlyphLoader(Face->glyph);
 	TText3DGlyphContourNodeShared Root = GlyphLoader.GetContourList();
 
 	return Root;
@@ -404,59 +289,37 @@ TText3DGlyphContourNodeShared FCachedFontData::GetGlyphContours(uint32 InGlyphIn
 #endif
 }
 
-void FCachedFontData::PrintCache() const
+void FText3DFontFaceCache::PrintCache() const
 {
-	UE_LOG(LogText3D, Log, TEXT("\n== Cached typefaces for font: %s =="), *GetFontName());
-
-	for (const TPair<uint32, FTypefaceFontData>& TypefaceData : TypefaceFontDataMap)
+#if WITH_FREETYPE
+	if (IsValid())
 	{
-		UE_LOG(LogText3D, Log, TEXT("\t- %s"), *TypefaceData.Value.GetTypefaceName().ToString());
+		const FT_Face Face = FontFace->GetFace();
+		UE_LOG(LogText3D, Log, TEXT("Cached Font Face Family=%s Style=%s Hash=%u Meshes=%i Usage=%i"), *FString(Face->family_name), *FString(Face->style_name), FontFaceHash, GlyphMeshes.Num(), GetCacheRefCount());
 	}
+#endif
 }
 
-bool FCachedFontData::Cleanup()
+bool FText3DFontFaceCache::IsValid() const
 {
-	TArray<uint32> TypefaceIndexes;
-	TypefaceFontDataMap.GetKeys(TypefaceIndexes);
-
-	for (const uint32 TypefaceIndex : TypefaceIndexes)
-	{
-		CleanupTypeface(TypefaceIndex);
-	}
-
-	if (TypefaceFontDataMap.IsEmpty())
-	{
-		return true;
-	}
-
+#if WITH_FREETYPE
+	return FontFace.IsValid() && FontFace->GetFace() != nullptr;
+#else
 	return false;
+#endif
 }
 
-bool FCachedFontData::CleanupTypeface(uint32 InTypefaceIndex)
+uint32 FText3DFontFaceCache::GetCacheRefCount() const
 {
-	bool bRemoved = false;
-	if (TypefaceFontDataMap.Contains(InTypefaceIndex))
-	{
-		FTypefaceFontData& TypefaceFontData = TypefaceFontDataMap[InTypefaceIndex];
-
-		for (auto MeshesIterator = TypefaceFontData.GetMeshes().CreateIterator(); MeshesIterator; ++MeshesIterator)
-		{
-			if (MeshesIterator.Value().GetCacheCount() <= 1)
-			{
-				MeshesIterator.RemoveCurrent();
-			}
-		}
-
-		if (TypefaceFontData.GetCacheCount() <= 1)
-		{
-			bRemoved = true;
-		}
-	}
-
-	if (bRemoved)
+	uint32 MeshCount = 0;
+	for (const TPair<uint32, FText3DCachedMesh>& GlyphMeshPair : GlyphMeshes)
 	{
-		TypefaceFontDataMap.Remove(InTypefaceIndex);
+		MeshCount += GlyphMeshPair.Value.RefCount;
 	}
+	return MeshCount;
+}
 
-	return bRemoved;
+uint32 FText3DFontFaceCache::GetCacheGlyphCount() const
+{
+	return GlyphMeshes.Num();
 }
diff --git a/Engine/Plugins/Experimental/Text3D/Source/Text3D/Private/Subsystems/Text3DEngineSubsystem.h b/Engine/Plugins/Experimental/Text3D/Source/Text3D/Private/Subsystems/Text3DEngineSubsystem.h
--- a/Engine/Plugins/Experimental/Text3D/Source/Text3D/Private/Subsystems/Text3DEngineSubsystem.h
+++ b/Engine/Plugins/Experimental/Text3D/Source/Text3D/Private/Subsystems/Text3DEngineSubsystem.h
@@ -5,201 +5,69 @@
 #include "Containers/Ticker.h"
 #include "GeometryBuilders/Text3DGlyphContourNode.h"
 #include "Subsystems/EngineSubsystem.h"
-#include "Text3DModule.h"
 #include "Text3DTypes.h"
 
 #include "Text3DEngineSubsystem.generated.h"
 
+struct FSharedStruct;
 class FFreeTypeFace;
-class UDynamicMesh;
-class UFont;
-class UMaterial;
-class UStaticMesh;
-class UWorld;
-struct FFontFaceData;
 
 USTRUCT()
-struct FGlyphMeshParameters
+struct FText3DFontFaceCache
 {
 	GENERATED_BODY()
 
-	UPROPERTY()
-	float Extrude = 5.0f;
-
-	UPROPERTY()
-	float Bevel = 0.0f;
-
-	UPROPERTY()
-	EText3DBevelType BevelType = EText3DBevelType::Convex;
-
-	UPROPERTY()
-	int32 BevelSegments = 8;
-
-	UPROPERTY()
-	bool bOutline = false;
-
-	UPROPERTY()
-	float OutlineExpand = 0.5f;
-
-	UPROPERTY()
-	EText3DOutlineType OutlineType = EText3DOutlineType::Stroke;
-
-	UPROPERTY()
-	uint32 TypefaceIndex = 0;
-
-	UPROPERTY()
-	FVector PivotOffset = FVector::ZeroVector;
-
-	friend uint32 GetTypeHash(const FGlyphMeshParameters& InElement)
-	{
-		uint32 HashParameters = 0;
-		HashParameters = HashCombine(HashParameters, GetTypeHash(InElement.Extrude));
-		HashParameters = HashCombine(HashParameters, GetTypeHash(InElement.Bevel));
-		HashParameters = HashCombine(HashParameters, GetTypeHash(InElement.BevelType));
-		HashParameters = HashCombine(HashParameters, GetTypeHash(InElement.BevelSegments));
-		HashParameters = HashCombine(HashParameters, GetTypeHash(InElement.bOutline));
-		HashParameters = HashCombine(HashParameters, GetTypeHash(InElement.OutlineExpand));
-		HashParameters = HashCombine(HashParameters, GetTypeHash(InElement.OutlineType));
-		HashParameters = HashCombine(HashParameters, GetTypeHash(InElement.TypefaceIndex));
-		HashParameters = HashCombine(HashParameters, GetTypeHash(InElement.PivotOffset));
-		return HashParameters;
-	}
-};
-
-USTRUCT()
-struct FText3DCachedMesh
-{
-	GENERATED_BODY()
-
-	UPROPERTY()
-	TObjectPtr<UStaticMesh> StaticMesh;
-
-	UPROPERTY()
-	TObjectPtr<UDynamicMesh> DynamicMesh;
-
-	UPROPERTY()
-	FBox MeshBounds = FBox(ForceInitToZero);
-
-	UPROPERTY()
-	FVector MeshOffset = FVector::ZeroVector;
-};
-
-USTRUCT()
-struct FCachedFontMeshes
-{
-	GENERATED_BODY()
-
-	FCachedFontMeshes();
-
-	int32 GetCacheCount() const;
-	TSharedPtr<int32> GetCacheCounter();
-
-	UPROPERTY()
-	TMap<uint32, FText3DCachedMesh> Glyphs;
-
-private:
-	TSharedPtr<int32> CacheCounter;
-};
-
-USTRUCT()
-struct FTypefaceFontData
-{
-	GENERATED_BODY()
-
-	FTypefaceFontData();
+#if WITH_FREETYPE
+	static uint32 GetFontFaceHash(const TSharedPtr<FFreeTypeFace>& InFontFace);
+#endif
+	static uint32 GetGlyphMeshHash(uint32 InGlyphIndex, const FGlyphMeshParameters& InParameters);
 
-	int32 GetCacheCount() const
-	{
-		const int32 Count = CacheCounter.GetSharedReferenceCount();
-		return Count;
-	}
+	FText3DFontFaceCache() = default;
 
 #if WITH_FREETYPE
-	FT_Face GetTypeFace() const { return TypeFace; }
-	void SetTypeFace(FT_Face InTypeFace);
-
-	TSharedPtr<const FFontFaceData> GetTypeFaceData() const;
-	void SetTypeFaceData(const TSharedPtr<const FFontFaceData>& InTypeFaceData);
+	explicit FText3DFontFaceCache(const TSharedPtr<FFreeTypeFace>& InFontFace);
 #endif
 
-	FName GetTypefaceName() const { return TypefaceName; }
-	void SetTypefaceName(FName InTypefaceName)
+	FText3DCachedMesh* FindOrBuildGlyphMesh(uint32 InGlyphIndex, const FGlyphMeshParameters& InParameters);
+	const FText3DCachedMesh* FindGlyphMesh(uint32 InGlyphHash) const;
+	FText3DCachedMesh* FindGlyphMesh(uint32 InGlyphHash);
+	uint32 CleanupCache();
+	
+	void PrintCache() const;
+	bool IsValid() const;
+	uint32 GetCacheRefCount() const;
+	uint32 GetCacheGlyphCount() const;
+
+	bool operator==(const FText3DFontFaceCache& InOther) const
 	{
-		TypefaceName = InTypefaceName;
+		return FontFaceHash == InOther.FontFaceHash;
 	}
 
-	uint32 GetTypefaceFontDataHash() const { return TypefaceFontDataHash; }
-	void SetTypefaceFontDataHash(uint32 InDataHash)
+	bool operator!=(const FText3DFontFaceCache& InOther) const
 	{
-		TypefaceFontDataHash = InDataHash;
+		return !(*this == InOther);
 	}
 
-	TMap<uint32, FCachedFontMeshes>& GetMeshes() { return Meshes; }
-
-	TSharedPtr<int32> GetCacheCounter() { return CacheCounter; }
-
-	FCachedFontMeshes& FindOrAddMeshes(uint32 InHashParameters)
+	friend uint32 GetTypeHash(const FText3DFontFaceCache& InValue)
 	{
-		return Meshes.FindOrAdd(InHashParameters);
+		return InValue.FontFaceHash;
 	}
 
 private:
-	UPROPERTY()
-	TMap<uint32, FCachedFontMeshes> Meshes;
-
-	FName TypefaceName = NAME_None;
-	TSharedPtr<int32> CacheCounter;
-	uint32 TypefaceFontDataHash = 0;
+	TText3DGlyphContourNodeShared LoadGlyphContours(uint32 InGlyphIndex) const;
 
 #if WITH_FREETYPE
-	TWeakPtr<const FFontFaceData> TypeFaceDataWeak;
-	FT_Face TypeFace = nullptr;
+	/** Font face for a Font+Typeface */
+	TSharedPtr<FFreeTypeFace> FontFace;
 #endif
-};
 
-USTRUCT()
-struct FCachedFontData
-{
-	GENERATED_BODY()
-
-	FCachedFontData();
-	~FCachedFontData();
-
-#if WITH_FREETYPE
-	FT_Face GetFreeTypeFace(uint32 InTypefaceIndex);
-#endif
-
-	FString GetFontName() const;
-
-	UFont* GetFont() { return Font; }
-	void SetFont(UFont* InFont)
-	{
-		Font = InFont;
-	}
-
-	void LoadFreeTypeFace(uint32 InTypefaceIndex);
-	void ClearFreeTypeFace();
-
-	bool Cleanup();
-	bool CleanupTypeface(uint32 InTypefaceIndex);
-
-	uint32 GetTypefaceFontDataHash(int32 InTypefaceEntryIndex) const;
-	const FTypefaceFontData* GetTypefaceFontData(int32 InTypefaceEntryIndex) const;
-	FTypefaceFontData* GetTypefaceFontData(int32 InTypefaceEntryIndex);
-	TSharedPtr<int32> GetCacheCounter(int32 InTypefaceEntryIndex);
-	TSharedPtr<int32> GetMeshesCacheCounter(const FGlyphMeshParameters& InParameters);
-
-	FText3DCachedMesh* GetGlyphMesh(uint32 InGlyphIndex, const FGlyphMeshParameters& InParameters);
-	TText3DGlyphContourNodeShared GetGlyphContours(uint32 InGlyphIndex, int32 InTypefaceEntryIndex);
-
-	void PrintCache() const;
-
-private:
+	/** Glyph indexes to cached meshes */
 	UPROPERTY()
-	TObjectPtr<UFont> Font;
+	TMap<uint32, FText3DCachedMesh> GlyphMeshes;
 
+	/** Hash to uniquely identify this struct */
 	UPROPERTY()
-	TMap<uint32, FTypefaceFontData> TypefaceFontDataMap;
+	uint32 FontFaceHash = 0;
 };
 
 UCLASS()
@@ -210,10 +78,18 @@ class UText3DEngineSubsystem : public UEngineSubsystem, public FSelfRegisteringE
 public:
 	static UText3DEngineSubsystem* Get();
 
-	UText3DEngineSubsystem();
+	void ScheduleCacheCleanup();
+	void CancelCacheCleanup();
+
+#if WITH_FREETYPE
+	FText3DFontFaceCache* FindOrAddCachedFontFace(const TSharedPtr<FFreeTypeFace>& InFontFace);
+#endif
+
+	FText3DFontFaceCache* FindCachedFontFace(uint32 InFontFaceHash);
 
+private:
 	// ~Begin UEngineSubsystem
-	virtual void Initialize(FSubsystemCollectionBase& Collection) override;
+	virtual void Initialize(FSubsystemCollectionBase& InCollection) override;
 	virtual void Deinitialize() override;
 	// ~End UEngineSubsystem
 
@@ -221,20 +97,12 @@ public:
 	virtual bool Exec_Dev(UWorld* InWorld, const TCHAR* InCmd, FOutputDevice& InAr) override;
 	// ~End FSelfRegisteringExec
 
+	bool CleanupCache(float);
 	void PrintCache() const;
-
-	void Reset();
-	void Cleanup();
-	FCachedFontData& GetCachedFontData(UFont* InFont, int32 InTypefaceEntryIndex = 0);
-
-	UPROPERTY()
-	TObjectPtr<UMaterial> DefaultMaterial;
-
-private:
-	bool CleanupTimerCallback(float DeltaTime);
+	void ClearCache();
 
 	UPROPERTY()
-	TMap<uint32, FCachedFontData> CachedFonts;
+	TMap<uint32, FText3DFontFaceCache> CachedFontFaces;
 
-	FTSTicker::FDelegateHandle CleanupTickerHandle;
+	FTSTicker::FDelegateHandle CacheCleanupDelegate;
 };
diff --git a/Engine/Plugins/Experimental/Text3D/Source/Text3D/Private/Text3DComponent.cpp b/Engine/Plugins/Experimental/Text3D/Source/Text3D/Private/Text3DComponent.cpp
--- a/Engine/Plugins/Experimental/Text3D/Source/Text3D/Private/Text3DComponent.cpp
+++ b/Engine/Plugins/Experimental/Text3D/Source/Text3D/Private/Text3DComponent.cpp
@@ -5,8 +5,8 @@
 #include "Containers/Ticker.h"
 #include "Engine/Engine.h"
 #include "Engine/Font.h"
-#include "Engine/StaticMesh.h"
 #include "Extensions/Text3DDefaultCharacterExtension.h"
+#include "Extensions/Text3DDefaultFormatStyleExtension.h"
 #include "Extensions/Text3DDefaultGeometryExtension.h"
 #include "Extensions/Text3DDefaultLayoutExtension.h"
 #include "Extensions/Text3DDefaultMaterialExtension.h"
@@ -22,7 +22,6 @@
 #include "Renderers/StaticMeshes/Text3DStaticMeshesRenderer.h"
 #include "Renderers/Text3DRendererBase.h"
 #include "Settings/Text3DProjectSettings.h"
-#include "Styling/StyleDefaults.h"
 #include "Subsystems/Text3DEngineSubsystem.h"
 #include "Text3DComponentVersion.h"
 #include "UObject/ConstructorHelpers.h"
@@ -50,6 +49,7 @@ UText3DComponent::UText3DComponent()
 	TextRendererClass = UText3DStaticMeshesRenderer::StaticClass();
 
 	// Extensions
+	FormatStyleExtension = CreateDefaultSubobject<UText3DFormatStyleExtensionBase, UText3DDefaultFormatStyleExtension>(TEXT("FormatStyleExtension"));
 	MaterialExtension = CreateDefaultSubobject<UText3DMaterialExtensionBase, UText3DDefaultMaterialExtension>(TEXT("MaterialExtension"));
 	GeometryExtension = CreateDefaultSubobject<UText3DGeometryExtensionBase, UText3DDefaultGeometryExtension>(TEXT("GeometryExtension"));
 	LayoutExtension = CreateDefaultSubobject<UText3DLayoutExtensionBase, UText3DDefaultLayoutExtension>(TEXT("LayoutExtension"));
@@ -615,6 +615,16 @@ UText3DRendererBase* UText3DComponent::GetTextRenderer() const
 	return TextRenderer;
 }
 
+UText3DFormatStyleExtensionBase* UText3DComponent::GetFormatStyleExtension(TSubclassOf<UText3DFormatStyleExtensionBase> InExtensionClass)
+{
+	if (!FormatStyleExtension || !InExtensionClass.Get())
+	{
+		return nullptr;
+	}
+
+	return FormatStyleExtension->IsA(InExtensionClass.Get()) ? FormatStyleExtension : nullptr;
+}
+
 UText3DLayoutExtensionBase* UText3DComponent::GetLayoutExtension(TSubclassOf<UText3DLayoutExtensionBase> InExtensionClass)
 {
 	if (!LayoutExtension || !InExtensionClass.Get())
diff --git a/Engine/Plugins/Experimental/Text3D/Source/Text3D/Private/Text3DTypes.cpp b/Engine/Plugins/Experimental/Text3D/Source/Text3D/Private/Text3DTypes.cpp
new file mode 100644
--- /dev/null
+++ b/Engine/Plugins/Experimental/Text3D/Source/Text3D/Private/Text3DTypes.cpp
@@ -0,0 +1,90 @@
+// Copyright Epic Games, Inc. All Rights Reserved.
+
+#include "Text3DTypes.h"
+
+#include "Subsystems/Text3DEngineSubsystem.h"
+
+namespace UE::Text3D::Geometry
+{
+	FCachedFontFaceGlyphHandle::FCachedFontFaceGlyphHandle(FText3DFontFaceCache* InFontFaceCache, uint32 InGlyphIndex, const FGlyphMeshParameters& InParameters)
+	{
+		if (InFontFaceCache)
+		{
+			if (FText3DCachedMesh* CachedMesh = InFontFaceCache->FindOrBuildGlyphMesh(InGlyphIndex, InParameters))
+			{
+				FontFaceHash = GetTypeHash(*InFontFaceCache);
+				GlyphIndex = InGlyphIndex;
+				GlyphMeshHash = FText3DFontFaceCache::GetGlyphMeshHash(InGlyphIndex, InParameters);
+				CachedMesh->RefCount++;
+			}
+		}
+	}
+
+	FCachedFontFaceGlyphHandle::~FCachedFontFaceGlyphHandle()
+	{
+		Unset();
+	}
+
+	FCachedFontFaceGlyphHandle& FCachedFontFaceGlyphHandle::operator=(const FCachedFontFaceGlyphHandle& InOther)
+	{
+		Unset();
+		FontFaceHash = InOther.FontFaceHash;
+		GlyphIndex = InOther.GlyphIndex;
+		GlyphMeshHash = InOther.GlyphMeshHash;
+		if (FText3DCachedMesh* CachedMesh = Resolve())
+		{
+			CachedMesh->RefCount++;
+		}
+		return *this;
+	}
+
+	void FCachedFontFaceGlyphHandle::Unset()
+	{
+		if (FText3DCachedMesh* CachedMesh = Resolve())
+		{
+			CachedMesh->RefCount--;
+		}
+
+		FontFaceHash = 0;
+		GlyphMeshHash = 0;
+		GlyphIndex = 0;
+	}
+
+	const FText3DCachedMesh* FCachedFontFaceGlyphHandle::Resolve() const
+	{
+		if (IsValid())
+		{
+			if (UText3DEngineSubsystem* Text3DSubsystem = UText3DEngineSubsystem::Get())
+			{
+				if (const FText3DFontFaceCache* CachedFontFace = Text3DSubsystem->FindCachedFontFace(FontFaceHash))
+				{
+					if (const FText3DCachedMesh* CachedMesh = CachedFontFace->FindGlyphMesh(GlyphMeshHash))
+					{
+						return CachedMesh;
+					}
+				}
+			}
+		}
+
+		return nullptr;
+	}
+
+	FText3DCachedMesh* FCachedFontFaceGlyphHandle::Resolve()
+	{
+		if (IsValid())
+		{
+			if (UText3DEngineSubsystem* Text3DSubsystem = UText3DEngineSubsystem::Get())
+			{
+				if (FText3DFontFaceCache* CachedFontFace = Text3DSubsystem->FindCachedFontFace(FontFaceHash))
+				{
+					if (FText3DCachedMesh* CachedMesh = CachedFontFace->FindGlyphMesh(GlyphMeshHash))
+					{
+						return CachedMesh;
+					}
+				}
+			}
+		}
+
+		return nullptr;
+	}
+}
diff --git a/Engine/Plugins/Experimental/Text3D/Source/Text3D/Public/Characters/Text3DCharacterBase.h b/Engine/Plugins/Experimental/Text3D/Source/Text3D/Public/Characters/Text3DCharacterBase.h
--- a/Engine/Plugins/Experimental/Text3D/Source/Text3D/Public/Characters/Text3DCharacterBase.h
+++ b/Engine/Plugins/Experimental/Text3D/Source/Text3D/Public/Characters/Text3DCharacterBase.h
@@ -32,14 +32,12 @@ public:
 	}
 #endif
 
-	void SetGlyphIndex(uint32 InGlyphIndex);
-	uint32 GetGlyphIndex() const;
-
-	void SetMeshBounds(const FBox& InBounds);
-	const FBox& GetMeshBounds() const;
+	const UE::Text3D::Geometry::FCachedFontFaceGlyphHandle& GetFontFaceGlyphHandle() const;
+	void SetFontFaceGlyphHandle(const UE::Text3D::Geometry::FCachedFontFaceGlyphHandle& InHandle);
 
-	void SetMeshOffset(const FVector& InOffset);
-	const FVector& GetMeshOffset() const;
+	uint32 GetGlyphIndex() const;
+	const FText3DCachedMesh* GetGlyphMesh() const;
+	void GetGlyphMeshBoundsAndOffset(FBox& OutBounds, FVector& OutOffset) const;
 
 	UFUNCTION()
 	TEXT3D_API void SetRelativeLocation(const FVector& InLocation);
@@ -109,12 +107,6 @@ protected:
 	/** Final transform after all extensions are applied */
 	FTransform Transform = FTransform::Identity;
 
-	/** Glyph that represents this character */
-	uint32 GlyphIndex = 0;
-
-	/** Actual mesh render bounds */
-	FBox MeshBounds = FBox(ForceInitToZero);
-
-	/** Offset around mesh due to font face */
-	FVector MeshOffset = FVector::ZeroVector;
+	/** Handle to the current glyph mesh this character represents */
+	UE::Text3D::Geometry::FCachedFontFaceGlyphHandle FontFaceGlyphHandle;
 };
diff --git a/Engine/Plugins/Experimental/Text3D/Source/Text3D/Public/Extensions/Text3DDefaultFormatStyleExtension.h b/Engine/Plugins/Experimental/Text3D/Source/Text3D/Public/Extensions/Text3DDefaultFormatStyleExtension.h
new file mode 100644
--- /dev/null
+++ b/Engine/Plugins/Experimental/Text3D/Source/Text3D/Public/Extensions/Text3DDefaultFormatStyleExtension.h
@@ -0,0 +1,41 @@
+// Copyright Epic Games, Inc. All Rights Reserved.
+
+#pragma once
+
+#include "Containers/Set.h"
+#include "Text3DFormatStyleExtensionBase.h"
+#include "Text3DDefaultFormatStyleExtension.generated.h"
+
+class FSlateStyleSet;
+class UText3DFormatStyleBase;
+struct FTextBlockStyle;
+
+UCLASS(MinimalAPI)
+class UText3DDefaultFormatStyleExtension : public UText3DFormatStyleExtensionBase
+{
+	GENERATED_BODY()
+	
+protected:
+	//~ Begin UObject
+#if WITH_EDITOR
+	virtual void PostEditChangeProperty(FPropertyChangedEvent& InPropertyChangedEvent) override;
+#endif
+	//~ End UObject
+
+	//~ Begin UText3DFormatStyleExtensionBase
+	virtual TSharedPtr<FTextBlockStyle> GetDefaultStyle() const override;
+    virtual TSharedPtr<FSlateStyleSet> GetCustomStyles() const override;
+    //~ End UText3DFormatStyleExtensionBase
+    
+    //~ Begin UText3DExtensionBase
+    virtual EText3DExtensionResult PreRendererUpdate(const UE::Text3D::Renderer::FUpdateParameters& InParameters) override;
+    virtual EText3DExtensionResult PostRendererUpdate(const UE::Text3D::Renderer::FUpdateParameters& InParameters) override;
+    //~ End UText3DExtensionBase
+
+	UPROPERTY(EditInstanceOnly, Instanced, Category="Style", meta=(AllowPrivateAccess="true"))
+	TSet<TObjectPtr<UText3DFormatStyleBase>> FormatStyles;
+
+private:
+	TSharedPtr<FTextBlockStyle> DefaultFontStyle;
+    TSharedPtr<FSlateStyleSet> TextStyleSet;
+};
diff --git a/Engine/Plugins/Experimental/Text3D/Source/Text3D/Public/Extensions/Text3DDefaultGeometryExtension.h b/Engine/Plugins/Experimental/Text3D/Source/Text3D/Public/Extensions/Text3DDefaultGeometryExtension.h
--- a/Engine/Plugins/Experimental/Text3D/Source/Text3D/Public/Extensions/Text3DDefaultGeometryExtension.h
+++ b/Engine/Plugins/Experimental/Text3D/Source/Text3D/Public/Extensions/Text3DDefaultGeometryExtension.h
@@ -114,7 +114,7 @@ protected:
 	virtual const FText3DCachedMesh* FindOrLoadGlyphMesh(UText3DCharacterBase* InCharacter) const override;
 	virtual EText3DHorizontalTextAlignment GetGlyphHAlignment() const override;
 	virtual EText3DVerticalTextAlignment GetGlyphVAlignment() const override;
-	virtual FTypefaceFontData* GetTypefaceFontData() const override;
+	virtual const FGlyphMeshParameters* GetGlyphMeshParameters() const override;
 	//~ End UText3DGeometryExtensionBase
 
 	//~ Begin UText3DExtensionBase
@@ -164,6 +164,5 @@ protected:
 	EText3DVerticalTextAlignment PivotVAlignment = EText3DVerticalTextAlignment::Bottom;
 
 private:
-	FCachedFontData* CachedFontData = nullptr;
 	FGlyphMeshParameters GlyphMeshParameters;
 };
\ No newline at end of file
diff --git a/Engine/Plugins/Experimental/Text3D/Source/Text3D/Public/Extensions/Text3DDefaultLayoutExtension.h b/Engine/Plugins/Experimental/Text3D/Source/Text3D/Public/Extensions/Text3DDefaultLayoutExtension.h
--- a/Engine/Plugins/Experimental/Text3D/Source/Text3D/Public/Extensions/Text3DDefaultLayoutExtension.h
+++ b/Engine/Plugins/Experimental/Text3D/Source/Text3D/Public/Extensions/Text3DDefaultLayoutExtension.h
@@ -3,13 +3,10 @@
 #pragma once
 
 #include "Extensions/Text3DLayoutExtensionBase.h"
-#include "Framework/Text/ITextLayoutMarshaller.h"
 #include "Text3DDefaultLayoutExtension.generated.h"
 
+class FRichTextLayoutMarshaller;
 class FText3DLayout;
-struct FPropertyChangedEvent;
-struct FTextBlockStyle;
-struct FText3DTypeFaceMetrics;
 
 UCLASS(MinimalAPI)
 class UText3DDefaultLayoutExtension : public UText3DLayoutExtensionBase
@@ -219,10 +216,6 @@ private:
 	/** Stores the text layout calculated by the TextLayoutMarshaller. */
 	TSharedPtr<FText3DLayout> TextLayout;
 
-	TSharedPtr<FTextBlockStyle> FontStyle;
-
 	/** Determines how text is laid out, ie. parsing line breaks. */
-	TSharedPtr<ITextLayoutMarshaller> TextLayoutMarshaller;
-
-	FText3DTypeFaceMetrics TypeFaceMetrics;
+	TSharedPtr<FRichTextLayoutMarshaller> TextLayoutMarshaller;
 };
\ No newline at end of file
diff --git a/Engine/Plugins/Experimental/Text3D/Source/Text3D/Public/Extensions/Text3DFormatStyleExtensionBase.h b/Engine/Plugins/Experimental/Text3D/Source/Text3D/Public/Extensions/Text3DFormatStyleExtensionBase.h
new file mode 100644
--- /dev/null
+++ b/Engine/Plugins/Experimental/Text3D/Source/Text3D/Public/Extensions/Text3DFormatStyleExtensionBase.h
@@ -0,0 +1,36 @@
+// Copyright Epic Games, Inc. All Rights Reserved.
+
+#pragma once
+
+#include "Text3DExtensionBase.h"
+#include "Text3DTypes.h"
+#include "Templates/SharedPointer.h"
+#include "Text3DFormatStyleExtensionBase.generated.h"
+
+class FSlateStyleSet;
+struct FText3DShapedGlyphText;
+struct FTextBlockStyle;
+
+/** Extension that handles formatting styles for Text3D */
+UCLASS(MinimalAPI, Abstract)
+class UText3DFormatStyleExtensionBase : public UText3DExtensionBase
+{
+	GENERATED_BODY()
+
+public:
+	UText3DFormatStyleExtensionBase()
+		: UText3DExtensionBase(/** Priority */1)
+	{}
+
+	/** Default style used by text when no other style is applied */
+	virtual TSharedPtr<FTextBlockStyle> GetDefaultStyle() const
+	{
+		return nullptr;
+	}
+
+	/** Custom defined styles to apply on text */
+	virtual TSharedPtr<FSlateStyleSet> GetCustomStyles() const
+	{
+		return nullptr;
+	}
+};
diff --git a/Engine/Plugins/Experimental/Text3D/Source/Text3D/Public/Extensions/Text3DGeometryExtensionBase.h b/Engine/Plugins/Experimental/Text3D/Source/Text3D/Public/Extensions/Text3DGeometryExtensionBase.h
--- a/Engine/Plugins/Experimental/Text3D/Source/Text3D/Public/Extensions/Text3DGeometryExtensionBase.h
+++ b/Engine/Plugins/Experimental/Text3D/Source/Text3D/Public/Extensions/Text3DGeometryExtensionBase.h
@@ -23,6 +23,7 @@ public:
 		: UText3DExtensionBase(/** Priority */1)
 	{}
 
+	// Deprecated do not use anymore, use character getter directly
 	virtual const FText3DCachedMesh* FindOrLoadGlyphMesh(UText3DCharacterBase* InCharacter) const
 	{
 		return nullptr;
@@ -38,7 +39,7 @@ public:
 		return EText3DVerticalTextAlignment::Bottom;
 	}
 
-	virtual FTypefaceFontData* GetTypefaceFontData() const
+	virtual const FGlyphMeshParameters* GetGlyphMeshParameters() const
 	{
 		return nullptr;
 	}
diff --git a/Engine/Plugins/Experimental/Text3D/Source/Text3D/Public/Extensions/Text3DLayoutTransformEffect.h b/Engine/Plugins/Experimental/Text3D/Source/Text3D/Public/Extensions/Text3DLayoutTransformEffect.h
--- a/Engine/Plugins/Experimental/Text3D/Source/Text3D/Public/Extensions/Text3DLayoutTransformEffect.h
+++ b/Engine/Plugins/Experimental/Text3D/Source/Text3D/Public/Extensions/Text3DLayoutTransformEffect.h
@@ -9,7 +9,7 @@
 class UCurveFloat;
 
 /** Extension that handles transform data for Text3D */
-UCLASS(MinimalAPI, EditInlineNew, ClassGroup=Text3D, DisplayName="TransformEffect")
+UCLASS(MinimalAPI, EditInlineNew, ClassGroup=Text3D, DisplayName="TransformEffect", AutoExpandCategories=(Transform))
 class UText3DLayoutTransformEffect : public UText3DLayoutEffectBase
 {
 	GENERATED_BODY()
diff --git a/Engine/Plugins/Experimental/Text3D/Source/Text3D/Public/Settings/Text3DProjectSettings.h b/Engine/Plugins/Experimental/Text3D/Source/Text3D/Public/Settings/Text3DProjectSettings.h
--- a/Engine/Plugins/Experimental/Text3D/Source/Text3D/Public/Settings/Text3DProjectSettings.h
+++ b/Engine/Plugins/Experimental/Text3D/Source/Text3D/Public/Settings/Text3DProjectSettings.h
@@ -70,6 +70,12 @@ public:
 
 	TEXT3D_API UMaterial* GetBaseMaterial(const FText3DMaterialKey& InKey) const;
 
+	TEXT3D_API void SetScheduleFontFaceGlyphCleanup(bool bInEnabled);
+	TEXT3D_API bool GetScheduleFontFaceGlyphCleanup() const;
+
+	TEXT3D_API void SetFontFaceGlyphCleanupPeriod(float InValue);
+	TEXT3D_API float GetFontFaceGlyphCleanupPeriod() const;
+
 #if WITH_EDITOR
 	void OpenEditorSettingsWindow() const;
 	TEXT3D_API TConstArrayView<FString> GetFavoriteFonts() const;
@@ -88,6 +94,7 @@ public:
 
 protected:
 	void AddMaterial(FText3DMaterialKey&& InKey, TSoftObjectPtr<UMaterial>&& InMaterial);
+	void OnFontFaceCleanupOptionsChanged();
 
 #if WITH_EDITOR
 	//~ Begin UObject
@@ -113,6 +120,14 @@ protected:
 	UPROPERTY(Config, EditAnywhere, Category="Text3D")
 	FString FontDirectory;
 
+	/** Schedule font face glyph mesh cleanup */
+	UPROPERTY(Config, EditAnywhere, Category="Text3D")
+	bool bScheduleFontFaceGlyphCleanup = true;
+
+	/** Interval for the font face glyph mesh cleanup */
+	UPROPERTY(Config, EditAnywhere, Category="Text3D", meta=(ClampMin="0", Units=Seconds, EditCondition="bScheduleFontFaceGlyphCleanup"))
+	float FontFaceGlyphCleanupPeriod = 300.f;
+
 #if WITH_EDITORONLY_DATA
 	/** Favorites fonts pinned in the font viewer dropdown */
 	UPROPERTY(Config, EditAnywhere, Category="Text3D")
diff --git a/Engine/Plugins/Experimental/Text3D/Source/Text3D/Public/Styles/Text3DFormatStyleBase.h b/Engine/Plugins/Experimental/Text3D/Source/Text3D/Public/Styles/Text3DFormatStyleBase.h
new file mode 100644
--- /dev/null
+++ b/Engine/Plugins/Experimental/Text3D/Source/Text3D/Public/Styles/Text3DFormatStyleBase.h
@@ -0,0 +1,81 @@
+// Copyright Epic Games, Inc. All Rights Reserved.
+
+#pragma once
+
+#include "Containers/Array.h"
+#include "Styling/SlateTypes.h"
+#include "UObject/Object.h"
+#include "Text3DFormatStyleBase.generated.h"
+
+class UFont;
+struct FTextBlockStyle;
+struct FTypefaceEntry;
+
+/** Text3D rich text format styles */
+UCLASS(MinimalAPI, EditInlineNew, ClassGroup=Text3D, DisplayName="FormatStyle", AutoExpandCategories=(Style))
+class UText3DFormatStyleBase : public UObject
+{
+	GENERATED_BODY()
+
+public:
+	UText3DFormatStyleBase();
+
+	const FTextBlockStyle& GetTextStyle() const;
+
+	UFUNCTION(BlueprintCallable, Category = "Text3D|Style")
+	TEXT3D_API void SetStyleName(FName InName);
+
+	UFUNCTION(BlueprintPure, Category = "Text3D|Style")
+	FName GetStyleName() const
+	{
+		return StyleName;
+	}
+
+	UFUNCTION(BlueprintCallable, Category = "Text3D|Style")
+	TEXT3D_API void SetFont(UFont* InFont);
+
+	UFUNCTION(BlueprintPure, Category = "Text3D|Style")
+	UFont* GetFont() const
+	{
+		return Font;
+	}
+
+	UFUNCTION(BlueprintCallable, Category = "Text3D|Style")
+	TEXT3D_API void SetFontTypeface(FName InTypeface);
+
+	UFUNCTION(BlueprintPure, Category = "Text3D|Style")
+	FName GetFontTypeface() const
+	{
+		return FontTypeface;
+	}
+
+protected:
+	//~ Begin UObject
+	virtual void PostLoad() override;
+#if WITH_EDITOR
+	virtual void PostEditChangeProperty(FPropertyChangedEvent& InPropertyChangedEvent) override;
+#endif
+	//~ End UObject
+
+	void OnStylePropertiesChanged();
+	void UpdateTextStyle();
+	void OnFontChanged();
+
+	UFUNCTION()
+	TArray<FName> GetTypefaceNames() const;
+	TArrayView<const FTypefaceEntry> GetAvailableTypefaces() const;
+
+	/** Style name used in text */
+	UPROPERTY(EditAnywhere, Setter, Getter, Category="Style", meta=(AllowPrivateAccess="true"))
+	FName StyleName = NAME_None;
+
+	/** Text font defines the style of rendered characters */
+	UPROPERTY(EditAnywhere, Setter, Getter, Category="Style", meta=(AdvancedFontPicker, AllowPrivateAccess="true"))
+	TObjectPtr<UFont> Font;
+
+	/** Text font face, subset within font like bold, italic, regular */
+	UPROPERTY(EditAnywhere, Setter, Getter, Category="Style", meta=(GetOptions="GetTypefaceNames", AllowPrivateAccess="true"))
+	FName FontTypeface;
+
+	FTextBlockStyle TextStyle;
+};
diff --git a/Engine/Plugins/Experimental/Text3D/Source/Text3D/Public/Text3DComponent.h b/Engine/Plugins/Experimental/Text3D/Source/Text3D/Public/Text3DComponent.h
--- a/Engine/Plugins/Experimental/Text3D/Source/Text3D/Public/Text3DComponent.h
+++ b/Engine/Plugins/Experimental/Text3D/Source/Text3D/Public/Text3DComponent.h
@@ -5,21 +5,21 @@
 #include "Components/SceneComponent.h"
 #include "Containers/Ticker.h"
 #include "Text3DTypes.h"
-#include "UObject/ObjectMacros.h"
 #include "Text3DComponent.generated.h"
 
-class UText3DCharacterBase;
-class UText3DCharacterExtensionBase;
 class FTextLayout;
 class ITextLayoutMarshaller;
 class UFont;
 class UMaterialInterface;
-class UText3DLayoutExtensionBase;
+class UText3DCharacterBase;
+class UText3DCharacterExtensionBase;
+class UText3DFormatStyleExtensionBase;
 class UText3DGeometryExtensionBase;
+class UText3DLayoutEffectBase;
+class UText3DLayoutExtensionBase;
 class UText3DMaterialExtensionBase;
 class UText3DRendererBase;
 class UText3DRenderingExtensionBase;
-class UText3DLayoutEffectBase;
 struct FTypefaceEntry;
 
 UCLASS(MinimalAPI
@@ -268,6 +268,20 @@ public:
 	
 	TEXT3D_API UText3DRendererBase* GetTextRenderer() const;
 
+	UFUNCTION(BlueprintCallable, Category="Text3D", meta=(DeterminesOutputType="InExtensionClass"))
+	UText3DFormatStyleExtensionBase* GetFormatStyleExtension(TSubclassOf<UText3DFormatStyleExtensionBase> InExtensionClass);
+	
+	UText3DFormatStyleExtensionBase* GetFormatStyleExtension() const
+	{
+		return FormatStyleExtension;
+	}
+	
+	template <typename T UE_REQUIRES(TIsDerivedFrom<T, UText3DFormatStyleExtensionBase>::Value)>
+	T* GetCastedFormatStyleExtension() const
+	{
+		return Cast<T>(FormatStyleExtension);
+	}
+
 	UFUNCTION(BlueprintCallable, Category="Text3D", meta=(DeterminesOutputType="InExtensionClass"))
 	UText3DLayoutExtensionBase* GetLayoutExtension(TSubclassOf<UText3DLayoutExtensionBase> InExtensionClass);
 
@@ -276,10 +290,10 @@ public:
 		return LayoutExtension;
 	}
 
-	template <typename InType UE_REQUIRES(TIsDerivedFrom<InType, UText3DLayoutExtensionBase>::Value)>
-	InType* GetCastedLayoutExtension() const
+	template <typename T UE_REQUIRES(TIsDerivedFrom<T, UText3DLayoutExtensionBase>::Value)>
+	T* GetCastedLayoutExtension() const
 	{
-		return Cast<InType>(LayoutExtension);
+		return Cast<T>(LayoutExtension);
 	}
 
 	UFUNCTION(BlueprintCallable, Category="Text3D", meta=(DeterminesOutputType="InExtensionClass"))
@@ -290,10 +304,10 @@ public:
 		return MaterialExtension;
 	}
 
-	template <typename InType UE_REQUIRES(TIsDerivedFrom<InType, UText3DMaterialExtensionBase>::Value)>
-	InType* GetCastedMaterialExtension() const
+	template <typename T UE_REQUIRES(TIsDerivedFrom<T, UText3DMaterialExtensionBase>::Value)>
+	T* GetCastedMaterialExtension() const
 	{
-		return Cast<InType>(MaterialExtension);
+		return Cast<T>(MaterialExtension);
 	}
 
 	UFUNCTION(BlueprintCallable, Category="Text3D", meta=(DeterminesOutputType="InExtensionClass"))
@@ -304,10 +318,10 @@ public:
 		return GeometryExtension;
 	}
 
-	template <typename InType UE_REQUIRES(TIsDerivedFrom<InType, UText3DGeometryExtensionBase>::Value)>
-	InType* GetCastedGeometryExtension() const
+	template <typename T UE_REQUIRES(TIsDerivedFrom<T, UText3DGeometryExtensionBase>::Value)>
+	T* GetCastedGeometryExtension() const
 	{
-		return Cast<InType>(GeometryExtension);
+		return Cast<T>(GeometryExtension);
 	}
 
 	UFUNCTION(BlueprintCallable, Category="Text3D", meta=(DeterminesOutputType="InExtensionClass"))
@@ -318,10 +332,10 @@ public:
 		return RenderingExtension;
 	}
 	
-	template <typename InType UE_REQUIRES(TIsDerivedFrom<InType, UText3DRenderingExtensionBase>::Value)>
-	InType* GetCastedRenderingExtension() const
+	template <typename T UE_REQUIRES(TIsDerivedFrom<T, UText3DRenderingExtensionBase>::Value)>
+	T* GetCastedRenderingExtension() const
 	{
-		return Cast<InType>(RenderingExtension);
+		return Cast<T>(RenderingExtension);
 	}
 
 	UFUNCTION(BlueprintCallable, Category="Text3D", meta=(DeterminesOutputType="InEffectClass"))
@@ -337,30 +351,30 @@ public:
 		return CharacterExtension;
 	}
 	
-	template <typename InType UE_REQUIRES(TIsDerivedFrom<InType, UText3DCharacterExtensionBase>::Value)>
-	InType* GetCastedCharacterExtension() const
+	template <typename T UE_REQUIRES(TIsDerivedFrom<T, UText3DCharacterExtensionBase>::Value)>
+	T* GetCastedCharacterExtension() const
 	{
-		return Cast<InType>(CharacterExtension);
+		return Cast<T>(CharacterExtension);
 	}
 
 	TEXT3D_API uint16 GetCharacterCount() const;
 
 	TEXT3D_API UText3DCharacterBase* GetCharacter(uint16 InCharacterIndex) const;
 
-	template <typename InType UE_REQUIRES(TIsDerivedFrom<InType, UText3DCharacterBase>::Value)>
-	InType* GetCastedCharacter(uint16 InCharacterIndex) const
+	template <typename T UE_REQUIRES(TIsDerivedFrom<T, UText3DCharacterBase>::Value)>
+	T* GetCastedCharacter(uint16 InCharacterIndex) const
 	{
-		return Cast<InType>(GetCharacter(InCharacterIndex));
+		return Cast<T>(GetCharacter(InCharacterIndex));
 	}
 
 	TEXT3D_API void ForEachCharacter(const TFunctionRef<void(UText3DCharacterBase*, uint16, uint16)>& InFunctor) const;
 
-	template <typename InType UE_REQUIRES(TIsDerivedFrom<InType, UText3DCharacterBase>::Value)>
-	void ForEachCastedCharacter(const TFunctionRef<void(InType*, uint16, uint16)>& InFunctor) const
+	template <typename T UE_REQUIRES(TIsDerivedFrom<T, UText3DCharacterBase>::Value)>
+	void ForEachCastedCharacter(const TFunctionRef<void(T*, uint16, uint16)>& InFunctor) const
 	{
 		ForEachCharacter([&InFunctor](UText3DCharacterBase* InCharacter, uint16 InCharacterIndex, uint16 InCharacterCount)
 		{
-			InFunctor(Cast<InType>(InCharacter), InCharacterIndex, InCharacterCount);
+			InFunctor(Cast<T>(InCharacter), InCharacterIndex, InCharacterCount);
 		});
 	}
 
@@ -436,17 +450,20 @@ protected:
 	TOptional<FText> FormattedText;
 
 	/** Text font defines the style of rendered characters */
-	UPROPERTY(EditAnywhere, Getter, Setter, Category = "Text", meta = (AdvancedFontPicker, AllowPrivateAccess = "true"))
+	UPROPERTY(EditAnywhere, Getter, Setter, DisplayName = "Default Font", Category = "Text", meta = (AdvancedFontPicker, AllowPrivateAccess = "true"))
 	TObjectPtr<UFont> Font;
 
 	/** Text font face, subset within font like bold, italic, regular */
-	UPROPERTY(EditAnywhere, Getter, Setter, Category = "Text", meta = (AllowPrivateAccess = "true", GetOptions="GetTypefaceNames"))
+	UPROPERTY(EditAnywhere, Getter, Setter, DisplayName = "Default Typeface", Category = "Text", meta = (AllowPrivateAccess = "true", GetOptions="GetTypefaceNames"))
 	FName Typeface;
 
 	/** Whether to force upper case for text */
 	UPROPERTY(EditAnywhere, Getter = "GetEnforceUpperCase", Setter = "SetEnforceUpperCase", Category = "Text", meta = (AllowPrivateAccess="true"))
 	bool bEnforceUpperCase = false;
 
+	UPROPERTY(VisibleInstanceOnly, Instanced, Category="Style", meta=(AllowPrivateAccess="true"))
+	TObjectPtr<UText3DFormatStyleExtensionBase> FormatStyleExtension;
+
 	UPROPERTY(VisibleInstanceOnly, Instanced, Category="Character", meta=(AllowPrivateAccess="true"))
 	TObjectPtr<UText3DCharacterExtensionBase> CharacterExtension;
 
diff --git a/Engine/Plugins/Experimental/Text3D/Source/Text3D/Public/Text3DTypes.h b/Engine/Plugins/Experimental/Text3D/Source/Text3D/Public/Text3DTypes.h
--- a/Engine/Plugins/Experimental/Text3D/Source/Text3D/Public/Text3DTypes.h
+++ b/Engine/Plugins/Experimental/Text3D/Source/Text3D/Public/Text3DTypes.h
@@ -5,6 +5,11 @@
 #include "Misc/EnumClassFlags.h"
 #include "Text3DTypes.generated.h"
 
+class UDynamicMesh;
+class UStaticMesh;
+class UWorld;
+struct FText3DFontFaceCache;
+
 /** Enumerate Text3D update flags based on their priority */
 enum class EText3DRendererFlags : uint8
 {
@@ -220,13 +225,69 @@ struct FText3DMaterialGroupKey
 	EText3DMaterialStyle Style = EText3DMaterialStyle::Invalid;
 };
 
-struct FText3DTypeFaceMetrics
+USTRUCT()
+struct FGlyphMeshParameters
 {
-	float FontHeight = 0.f;
-	/** Height above baseline */
-	float FontAscender = 0.f;
-	/** Height below baseline */
-	float FontDescender = 0.f;
+	GENERATED_BODY()
+
+	UPROPERTY()
+	float Extrude = 5.0f;
+
+	UPROPERTY()
+	float Bevel = 0.0f;
+
+	UPROPERTY()
+	EText3DBevelType BevelType = EText3DBevelType::Convex;
+
+	UPROPERTY()
+	int32 BevelSegments = 8;
+
+	UPROPERTY()
+	bool bOutline = false;
+
+	UPROPERTY()
+	float OutlineExpand = 0.5f;
+
+	UPROPERTY()
+	EText3DOutlineType OutlineType = EText3DOutlineType::Stroke;
+
+	UPROPERTY()
+	FVector PivotOffset = FVector::ZeroVector;
+
+	friend uint32 GetTypeHash(const FGlyphMeshParameters& InElement)
+	{
+		uint32 HashParameters = 0;
+		HashParameters = HashCombine(HashParameters, GetTypeHash(InElement.Extrude));
+		HashParameters = HashCombine(HashParameters, GetTypeHash(InElement.Bevel));
+		HashParameters = HashCombine(HashParameters, GetTypeHash(InElement.BevelType));
+		HashParameters = HashCombine(HashParameters, GetTypeHash(InElement.BevelSegments));
+		HashParameters = HashCombine(HashParameters, GetTypeHash(InElement.bOutline));
+		HashParameters = HashCombine(HashParameters, GetTypeHash(InElement.OutlineExpand));
+		HashParameters = HashCombine(HashParameters, GetTypeHash(InElement.OutlineType));
+		HashParameters = HashCombine(HashParameters, GetTypeHash(InElement.PivotOffset));
+		return HashParameters;
+	}
+};
+
+USTRUCT()
+struct FText3DCachedMesh
+{
+	GENERATED_BODY()
+
+	UPROPERTY()
+	TObjectPtr<UStaticMesh> StaticMesh;
+
+	UPROPERTY()
+	TObjectPtr<UDynamicMesh> DynamicMesh;
+
+	UPROPERTY()
+	FBox MeshBounds = FBox(ForceInitToZero);
+
+	UPROPERTY()
+	FVector MeshOffset = FVector::ZeroVector;
+
+	/** Amount of active usage of this glyph, used for cleanup */
+	uint32 RefCount = 0;
 };
 
 namespace UE::Text3D
@@ -257,4 +318,37 @@ namespace UE::Text3D
 			bool bIsLastFlag = false;
 		};
 	}
+
+	namespace Geometry
+	{
+		/** Acts as a handle for FText3DCachedMesh */
+		struct FCachedFontFaceGlyphHandle
+		{
+			FCachedFontFaceGlyphHandle() = default;
+			FCachedFontFaceGlyphHandle(FText3DFontFaceCache* InFontFaceCache, uint32 InGlyphIndex, const FGlyphMeshParameters& InParameters);
+			~FCachedFontFaceGlyphHandle();
+
+			FCachedFontFaceGlyphHandle& operator=(const FCachedFontFaceGlyphHandle& InOther);
+
+			void Unset();
+
+			const FText3DCachedMesh* Resolve() const;
+			FText3DCachedMesh* Resolve();
+
+			bool IsValid() const
+			{
+				return FontFaceHash != 0 && GlyphMeshHash != 0;
+			}
+
+			uint32 GetGlyphIndex() const
+			{
+				 return GlyphIndex;
+			}
+
+		private:
+			uint32 FontFaceHash = 0;
+			uint32 GlyphIndex = 0;
+			uint32 GlyphMeshHash = 0;
+		};
+	}
 }
\ No newline at end of file
diff --git a/Engine/Plugins/Experimental/Text3D/Source/Text3DEditor/Private/Customizations/Text3DEditorTextComponentDetailCustomization.cpp b/Engine/Plugins/Experimental/Text3D/Source/Text3DEditor/Private/Customizations/Text3DEditorTextComponentDetailCustomization.cpp
--- a/Engine/Plugins/Experimental/Text3D/Source/Text3DEditor/Private/Customizations/Text3DEditorTextComponentDetailCustomization.cpp
+++ b/Engine/Plugins/Experimental/Text3D/Source/Text3DEditor/Private/Customizations/Text3DEditorTextComponentDetailCustomization.cpp
@@ -19,6 +19,26 @@ namespace UE::Text3DEditor::Customization
 		FPropertyEditorModule& PropertyModule = FModuleManager::GetModuleChecked<FPropertyEditorModule>(TEXT("PropertyEditor"));
 		const FName ComponentClassName = UText3DComponent::StaticClass()->GetFName();
 
+		const TArray<FName> CategoryNames =
+		{
+			TEXT("Text"),
+			TEXT("Style"),
+			TEXT("Character"),
+			TEXT("Layout"),
+			TEXT("LayoutEffects"),
+			TEXT("Geometry"),
+			TEXT("Material"),
+			TEXT("Rendering"),
+			TEXT("Utilities")
+		};
+
+		int32 SortOrder = 1100;
+		for (FName CategoryName : CategoryNames)
+		{
+			InDetailBuilder.EditCategory(CategoryName).SetSortOrder(SortOrder);
+			SortOrder += 100;
+		}
+
 		static bool bSectionInitialized = false;
 		if (!bSectionInitialized)
 		{
@@ -44,6 +64,7 @@ namespace UE::Text3DEditor::Customization
 
 			const FName StyleSectionName = TEXT("Style");
 			const TSharedRef<FPropertySection> MaterialSection = PropertyModule.FindOrCreateSection(ComponentClassName, StyleSectionName, FText::FromName(StyleSectionName));
+			MaterialSection->AddCategory(TEXT("Style"));
 			MaterialSection->AddCategory(TEXT("Material"));
 
 			const FName EffectSectionName = TEXT("Effects");
diff --git a/Engine/Plugins/Experimental/Text3D/Text3D.uplugin b/Engine/Plugins/Experimental/Text3D/Text3D.uplugin
--- a/Engine/Plugins/Experimental/Text3D/Text3D.uplugin
+++ b/Engine/Plugins/Experimental/Text3D/Text3D.uplugin
@@ -16,8 +16,7 @@
     {
       "Name": "Text3D",
       "Type": "Runtime",
-      "LoadingPhase": "Default",
-	  "TargetDenyList": [ "Server" ]
+      "LoadingPhase": "Default"
     },
 	{
 	  "Name":  "Text3DEditor",
@@ -28,13 +27,11 @@
   "Plugins": [
     {
       "Name": "GeometryProcessing",
-      "Enabled": true,
-	  "TargetDenyList": [ "Server" ]
+      "Enabled": true
     },
 	{
 	  "Name": "GeometryScripting", 
-	  "Enabled": true,
-      "TargetDenyList": [ "Server" ]
+	  "Enabled": true
 	},
     {
       "Name": "GeometryMask",
