Automate '3D Text - 3DText Font can be changed'

#rnx

diff --git a/Engine/Plugins/Experimental/Text3D/Source/Text3DEditor/Private/Widgets/SText3DEditorFontSelector.h b/Engine/Plugins/Experimental/Text3D/Source/Text3DEditor/Private/Widgets/SText3DEditorFontSelector.h
--- a/Engine/Plugins/Experimental/Text3D/Source/Text3DEditor/Private/Widgets/SText3DEditorFontSelector.h
+++ b/Engine/Plugins/Experimental/Text3D/Source/Text3DEditor/Private/Widgets/SText3DEditorFontSelector.h
@@ -43,7 +43,7 @@ protected:
 	void OnProjectFontItemSelectionChanged(TSharedPtr<FString> InItem, ESelectInfo::Type InSelectInfo);
 	void OnSystemFontItemSelectionChanged(TSharedPtr<FString> InItem, ESelectInfo::Type InSelectInfo);
 
-	void OnFontItemSelectionChanged(TSharedPtr<FString> InItem, ESelectInfo::Type InSelectInfo);
+	TEXT3DEDITOR_API void OnFontItemSelectionChanged(TSharedPtr<FString> InItem, ESelectInfo::Type InSelectInfo);
 	void OnPropertyResetToDefault();
 	void UpdateSelectedItem();
 
@@ -78,4 +78,6 @@ protected:
 	TSharedPtr<SComboButton> ComboButton;
 
 	TSharedPtr<IPropertyHandle> FontPropertyHandle;
+
+	friend class UText3DEditorAccessHelper;
 };
