MotionDesign : Text3D
- Added template to replace named variable with its content.Eg: {tempate} replaced by its content and {\template\} by {template} all in a single pass
- Fixed alignment code to better handle multiple fonts metrics
- Fixed glyph size issue due to font face cache invalidation

diff --git a/Engine/Plugins/Experimental/Text3D/Source/Text3D/Private/Extensions/Text3DDefaultFormatStyleExtension.cpp b/Engine/Plugins/Experimental/Text3D/Source/Text3D/Private/Extensions/Text3DDefaultFormatStyleExtension.cpp
--- a/Engine/Plugins/Experimental/Text3D/Source/Text3D/Private/Extensions/Text3DDefaultFormatStyleExtension.cpp
+++ b/Engine/Plugins/Experimental/Text3D/Source/Text3D/Private/Extensions/Text3DDefaultFormatStyleExtension.cpp
@@ -41,15 +41,18 @@ TSharedPtr<FSlateStyleSet> UText3DDefaultFormatStyleExtension::GetCustomStyles()
 
 EText3DExtensionResult UText3DDefaultFormatStyleExtension::PreRendererUpdate(const UE::Text3D::Renderer::FUpdateParameters& InParameters)
 {
-	if (!DefaultFontStyle.IsValid())
-	{
-		DefaultFontStyle = MakeShared<FTextBlockStyle>();
-	}
-
-	if (EnumHasAnyFlags(InParameters.UpdateFlags, EText3DRendererFlags::Geometry))
+	if (!DefaultFontStyle.IsValid()
+		|| !TextStyleSet.IsValid()
+		|| EnumHasAnyFlags(InParameters.UpdateFlags, EText3DRendererFlags::Geometry))
 	{
 		const UText3DComponent* Text3DComponent = GetText3DComponent();
 		const UFont* Font = Text3DComponent->GetFont();
+
+		if (!DefaultFontStyle.IsValid())
+		{
+			DefaultFontStyle = MakeShared<FTextBlockStyle>();
+		}
+
 		FSlateFontInfo FontInfo(Font, UE::Text3D::Metrics::SlateFontSize);
 		FontInfo.CompositeFont = Font && Font->GetCompositeFont() ? MakeShared<FCompositeFont>(*Font->GetCompositeFont()) : FStyleDefaults::GetFontInfo().CompositeFont;
 		FontInfo.TypefaceFontName = Text3DComponent->GetTypeface();
@@ -57,7 +60,7 @@ EText3DExtensionResult UText3DDefaultFormatStyleExtension::PreRendererUpdate(con
 		DefaultFontStyle->SetFont(FontInfo);
 
 		TextStyleSet = MakeShared<FSlateStyleSet>(TEXT("Text3DRichTextStyle"));
-		TextStyleSet->Set(FName(TEXT("Default")), *DefaultFontStyle);
+		TextStyleSet->Set(FName(TEXT("default")), *DefaultFontStyle);
 		for (const UText3DFormatStyleBase* FormatStyle : FormatStyles)
 		{
 			if (FormatStyle && !FormatStyle->GetStyleName().IsNone())
diff --git a/Engine/Plugins/Experimental/Text3D/Source/Text3D/Private/Extensions/Text3DDefaultGeometryExtension.cpp b/Engine/Plugins/Experimental/Text3D/Source/Text3D/Private/Extensions/Text3DDefaultGeometryExtension.cpp
--- a/Engine/Plugins/Experimental/Text3D/Source/Text3D/Private/Extensions/Text3DDefaultGeometryExtension.cpp
+++ b/Engine/Plugins/Experimental/Text3D/Source/Text3D/Private/Extensions/Text3DDefaultGeometryExtension.cpp
@@ -62,7 +62,7 @@ void UText3DDefaultGeometryExtension::SetPivotVAlignment(EText3DVerticalTextAlig
 	}
 
 	// Only one supported for now
-	if (PivotVAlignment != EText3DVerticalTextAlignment::Bottom)
+	if (InPivot != EText3DVerticalTextAlignment::Bottom)
 	{
 		return;
 	}
diff --git a/Engine/Plugins/Experimental/Text3D/Source/Text3D/Private/Extensions/Text3DDefaultLayoutExtension.cpp b/Engine/Plugins/Experimental/Text3D/Source/Text3D/Private/Extensions/Text3DDefaultLayoutExtension.cpp
--- a/Engine/Plugins/Experimental/Text3D/Source/Text3D/Private/Extensions/Text3DDefaultLayoutExtension.cpp
+++ b/Engine/Plugins/Experimental/Text3D/Source/Text3D/Private/Extensions/Text3DDefaultLayoutExtension.cpp
@@ -422,7 +422,7 @@ FVector UText3DDefaultLayoutExtension::GetLineLocation(int32 LineIndex) const
 		}
 		else if (VerticalAlignment == EText3DVerticalTextAlignment::Bottom)
 		{
-			VerticalOffset += TotalHeight + ShapedLine.MaxFontDescender;
+			VerticalOffset += TotalHeight - ShapedLine.MaxFontDescender;
 		}
 	}
 
diff --git a/Engine/Plugins/Experimental/Text3D/Source/Text3D/Private/Extensions/Text3DDefaultTemplateExtension.cpp b/Engine/Plugins/Experimental/Text3D/Source/Text3D/Private/Extensions/Text3DDefaultTemplateExtension.cpp
new file mode 100644
--- /dev/null
+++ b/Engine/Plugins/Experimental/Text3D/Source/Text3D/Private/Extensions/Text3DDefaultTemplateExtension.cpp
@@ -0,0 +1,64 @@
+// Copyright Epic Games, Inc. All Rights Reserved.
+
+#include "Extensions/Text3DDefaultTemplateExtension.h"
+
+#include "Internationalization/TextFormatter.h"
+#include "Text3DComponent.h"
+#include "Templates/Text3DTemplateBase.h"
+
+#if WITH_EDITOR
+void UText3DDefaultTemplateExtension::PostEditChangeProperty(FPropertyChangedEvent& InPropertyChangedEvent)
+{
+	Super::PostEditChangeProperty(InPropertyChangedEvent);
+
+	const FName MemberPropertyName = InPropertyChangedEvent.GetMemberPropertyName();
+
+	if (MemberPropertyName == GET_MEMBER_NAME_CHECKED(UText3DDefaultTemplateExtension, Templates))
+	{
+		if (InPropertyChangedEvent.ChangeType == EPropertyChangeType::ArrayClear
+			|| InPropertyChangedEvent.ChangeType == EPropertyChangeType::ArrayRemove
+			|| InPropertyChangedEvent.ChangeType == EPropertyChangeType::ArrayMove)
+		{
+			if (UText3DComponent* Text3DComponent = GetText3DComponent())
+			{
+				Text3DComponent->RequestUpdate(EText3DRendererFlags::All, /** Immediate */false);
+			}
+		}
+	}
+}
+#endif
+
+const FText& UText3DDefaultTemplateExtension::GetFormattedText() const
+{
+	return FormattedText;
+}
+
+EText3DExtensionResult UText3DDefaultTemplateExtension::PreRendererUpdate(const UE::Text3D::Renderer::FUpdateParameters& InParameters)
+{
+	const UText3DComponent* Text3DComponent = GetText3DComponent();
+	FormattedText = Text3DComponent->GetText();
+
+	FFormatNamedArguments NamedArguments;
+	NamedArguments.Reserve(Templates.Num() * 2);
+	for (UText3DTemplateBase* Template : Templates)
+	{
+		if (Template)
+		{
+			Template->CollectTemplates(NamedArguments);
+		}
+	}
+
+	constexpr bool bRebuildText = false;
+	constexpr bool bRebuildAsSource = false;
+	FormattedText = FTextFormatter::Format(MoveTemp(FormattedText), MoveTemp(NamedArguments), bRebuildText, bRebuildAsSource);
+
+	// For backward compatibility
+	Text3DComponent->FormatText(FormattedText);
+
+	return EText3DExtensionResult::Finished;
+}
+
+EText3DExtensionResult UText3DDefaultTemplateExtension::PostRendererUpdate(const UE::Text3D::Renderer::FUpdateParameters& InParameters)
+{
+	return EText3DExtensionResult::Active;
+}
diff --git a/Engine/Plugins/Experimental/Text3D/Source/Text3D/Private/LayoutBuilders/Text3DLayoutShaper.cpp b/Engine/Plugins/Experimental/Text3D/Source/Text3D/Private/LayoutBuilders/Text3DLayoutShaper.cpp
--- a/Engine/Plugins/Experimental/Text3D/Source/Text3D/Private/LayoutBuilders/Text3DLayoutShaper.cpp
+++ b/Engine/Plugins/Experimental/Text3D/Source/Text3D/Private/LayoutBuilders/Text3DLayoutShaper.cpp
@@ -35,6 +35,8 @@ void FText3DLayoutShaper::ShapeBidirectionalText(
 		return;
 	}
 
+	checkf(InStyles.IsValid(), TEXT("Styles set is not valid, cannot proceed"))
+
 	TextLayout->ClearLines();
 	TextLayout->ClearLineHighlights();
 	TextLayout->ClearRunRenderers();
diff --git a/Engine/Plugins/Experimental/Text3D/Source/Text3D/Private/LayoutBuilders/Text3DShapedGlyphLine.h b/Engine/Plugins/Experimental/Text3D/Source/Text3D/Private/LayoutBuilders/Text3DShapedGlyphLine.h
--- a/Engine/Plugins/Experimental/Text3D/Source/Text3D/Private/LayoutBuilders/Text3DShapedGlyphLine.h
+++ b/Engine/Plugins/Experimental/Text3D/Source/Text3D/Private/LayoutBuilders/Text3DShapedGlyphLine.h
@@ -24,12 +24,12 @@ struct FShapedGlyphLine
 	/** Direction of the text in the line */
 	TextBiDi::ETextDirection TextDirection = TextBiDi::ETextDirection::LeftToRight;
 
-	/** Max font height for this line */
+	/** Max font height for this line calculated based on MaxFontAscender and MaxFontDescender */
 	float MaxFontHeight = 0.f;
 
-	/** Max height above baseline */
+	/** Positive value representing max height above baseline */
 	float MaxFontAscender = 0.f;
 
-	/** Max height below baseline */
+ 	/** Positive value representing max height below baseline */
 	float MaxFontDescender = 0.f;
 };
diff --git a/Engine/Plugins/Experimental/Text3D/Source/Text3D/Private/LayoutBuilders/Text3DShapedGlyphText.cpp b/Engine/Plugins/Experimental/Text3D/Source/Text3D/Private/LayoutBuilders/Text3DShapedGlyphText.cpp
--- a/Engine/Plugins/Experimental/Text3D/Source/Text3D/Private/LayoutBuilders/Text3DShapedGlyphText.cpp
+++ b/Engine/Plugins/Experimental/Text3D/Source/Text3D/Private/LayoutBuilders/Text3DShapedGlyphText.cpp
@@ -48,7 +48,6 @@ void FText3DShapedGlyphText::CalculateWidth()
         TArray<FShapedGlyphEntry> CurrentWord;
         float LineWidth = 0.0f;
         float CurrentWordLength = 0.0f;
-    	float MaxFontHeight = 0.0f;
     	float MaxFontAscender = 0.0f;
     	float MaxFontDescender = 0.0f;
 
@@ -70,10 +69,11 @@ void FText3DShapedGlyphText::CalculateWidth()
 				&& CurrentWordLength != LineWidth)  // and the line is not just a single word that we can't break
     		{
     			CurrentLine->Width = LineWidth - CurrentWordLength;
-    			CurrentLine->MaxFontHeight = MaxFontHeight;
+    			CurrentLine->MaxFontHeight = MaxFontAscender + MaxFontDescender;
     			CurrentLine->MaxFontAscender = MaxFontAscender;
     			CurrentLine->MaxFontDescender = MaxFontDescender;
-    			MaxFontHeight = 0.0f;
+
+    			// Reset
     			MaxFontAscender = 0.0f;
     			MaxFontDescender = 0.0f;
 
@@ -89,9 +89,12 @@ void FText3DShapedGlyphText::CalculateWidth()
 			{
 #if WITH_FREETYPE
 				const FT_Face Face = FontFaces[Index]->GetFace();
-				MaxFontHeight = FMath::Max(MaxFontHeight, Face->size->metrics.height * UE::Text3D::Metrics::FontSizeInverse);
-				MaxFontAscender = FMath::Max(MaxFontAscender, Face->size->metrics.ascender * UE::Text3D::Metrics::FontSizeInverse);
-				MaxFontDescender = FMath::Max(MaxFontDescender, Face->size->metrics.descender * UE::Text3D::Metrics::FontSizeInverse);
+
+				const float Ascender = FMath::Max(0.0f, Face->size->metrics.ascender * UE::Text3D::Metrics::FontSizeInverse);
+				const float Descender = FMath::Max(0.0f, -Face->size->metrics.descender * UE::Text3D::Metrics::FontSizeInverse);
+
+				MaxFontAscender  = FMath::Max(MaxFontAscender,  Ascender);
+				MaxFontDescender = FMath::Max(MaxFontDescender, Descender);
 #endif
 			}
 
@@ -131,7 +134,7 @@ void FText3DShapedGlyphText::CalculateWidth()
         }
 
         CurrentLine->Width = LineWidth;
-    	CurrentLine->MaxFontHeight = MaxFontHeight;
+    	CurrentLine->MaxFontHeight = MaxFontAscender + MaxFontDescender;
     	CurrentLine->MaxFontAscender = MaxFontAscender;
     	CurrentLine->MaxFontDescender = MaxFontDescender;
     }
diff --git a/Engine/Plugins/Experimental/Text3D/Source/Text3D/Private/Renderers/DynamicMesh/Text3DDynamicMeshRenderer.cpp b/Engine/Plugins/Experimental/Text3D/Source/Text3D/Private/Renderers/DynamicMesh/Text3DDynamicMeshRenderer.cpp
--- a/Engine/Plugins/Experimental/Text3D/Source/Text3D/Private/Renderers/DynamicMesh/Text3DDynamicMeshRenderer.cpp
+++ b/Engine/Plugins/Experimental/Text3D/Source/Text3D/Private/Renderers/DynamicMesh/Text3DDynamicMeshRenderer.cpp
@@ -91,6 +91,10 @@ void UText3DDynamicMeshRenderer::OnUpdate(const UE::Text3D::Renderer::FUpdatePar
 			{
 				AppendGlyphMesh(MeshEditor, InIndex, InCharacter->GetGlyphIndex(), CachedGlyph->DynamicMesh);
 			}
+			else
+			{
+				ClearGlyphMesh(MeshEditor, InIndex);
+			}
 		});
 
 		RefreshBounds();
diff --git a/Engine/Plugins/Experimental/Text3D/Source/Text3D/Private/Renderers/StaticMeshes/Text3DStaticMeshesRenderer.cpp b/Engine/Plugins/Experimental/Text3D/Source/Text3D/Private/Renderers/StaticMeshes/Text3DStaticMeshesRenderer.cpp
--- a/Engine/Plugins/Experimental/Text3D/Source/Text3D/Private/Renderers/StaticMeshes/Text3DStaticMeshesRenderer.cpp
+++ b/Engine/Plugins/Experimental/Text3D/Source/Text3D/Private/Renderers/StaticMeshes/Text3DStaticMeshesRenderer.cpp
@@ -175,7 +175,7 @@ void UText3DStaticMeshesRenderer::OnDestroy()
 			MeshComponent->DestroyComponent();
 		}
 	}
-	CharacterMeshes.Reset();
+	CharacterMeshesPool.Reset();
 
 	if (TextRoot)
 	{
diff --git a/Engine/Plugins/Experimental/Text3D/Source/Text3D/Private/Renderers/Text3DRendererBase.cpp b/Engine/Plugins/Experimental/Text3D/Source/Text3D/Private/Renderers/Text3DRendererBase.cpp
--- a/Engine/Plugins/Experimental/Text3D/Source/Text3D/Private/Renderers/Text3DRendererBase.cpp
+++ b/Engine/Plugins/Experimental/Text3D/Source/Text3D/Private/Renderers/Text3DRendererBase.cpp
@@ -9,6 +9,7 @@
 #include "Extensions/Text3DLayoutExtensionBase.h"
 #include "Extensions/Text3DMaterialExtensionBase.h"
 #include "Extensions/Text3DRenderingExtensionBase.h"
+#include "Extensions/Text3DTemplateExtensionBase.h"
 #include "GameFramework/Actor.h"
 #include "Logs/Text3DLogs.h"
 #include "Settings/Text3DProjectSettings.h"
@@ -86,7 +87,9 @@ void UText3DRendererBase::Create()
 #endif
 
 	bInitialized = true;
-	UE_LOG(LogText3D, Log, TEXT("%s : Text3DRenderer %s Created"), *Text3DComponent->GetOwner()->GetActorNameOrLabel(), *GetFriendlyName().ToString())
+
+	const AActor* Owner = Text3DComponent->GetOwner();
+	UE_LOG(LogText3D, Log, TEXT("%s : Text3DRenderer %s Created"), !!Owner ? *Owner->GetActorNameOrLabel() : TEXT("Invalid owner"), *GetFriendlyName().ToString())
 }
 
 void UText3DRendererBase::Update(EText3DRendererFlags InFlags)
@@ -106,6 +109,7 @@ void UText3DRendererBase::Update(EText3DRendererFlags InFlags)
 
 	TArray<UText3DExtensionBase*, TInlineAllocator<6>> Extensions
 	{
+		Text3DComponent->GetTemplateExtension(),
 		Text3DComponent->GetFormatStyleExtension(),
 		Text3DComponent->GetCharacterExtension(),
 		Text3DComponent->GetGeometryExtension(),
@@ -187,7 +191,8 @@ void UText3DRendererBase::Update(EText3DRendererFlags InFlags)
 		Flag <<= 1;
 	}
 
-	UE_LOG(LogText3D, Verbose, TEXT("%s : Text3DRenderer %s Updated with flags %i"), *Text3DComponent->GetOwner()->GetActorNameOrLabel(), *GetFriendlyName().ToString(), InFlags)
+	const AActor* Owner = Text3DComponent->GetOwner();
+	UE_LOG(LogText3D, Verbose, TEXT("%s : Text3DRenderer %s Updated with flags %i"), !!Owner ? *Owner->GetActorNameOrLabel() : TEXT("Invalid owner"), *GetFriendlyName().ToString(), InFlags)
 }
 
 void UText3DRendererBase::Clear()
diff --git a/Engine/Plugins/Experimental/Text3D/Source/Text3D/Private/Subsystems/Text3DEngineSubsystem.cpp b/Engine/Plugins/Experimental/Text3D/Source/Text3D/Private/Subsystems/Text3DEngineSubsystem.cpp
--- a/Engine/Plugins/Experimental/Text3D/Source/Text3D/Private/Subsystems/Text3DEngineSubsystem.cpp
+++ b/Engine/Plugins/Experimental/Text3D/Source/Text3D/Private/Subsystems/Text3DEngineSubsystem.cpp
@@ -148,7 +148,18 @@ FText3DFontFaceCache* UText3DEngineSubsystem::FindOrAddCachedFontFace(const TSha
 	}
 
 	const uint32 FontFaceHash = FText3DFontFaceCache::GetFontFaceHash(InFontFace);
-	return FontFaceHash != 0 ? &CachedFontFaces.FindOrAdd(FontFaceHash, FText3DFontFaceCache(InFontFace)) : nullptr;
+
+	if (FontFaceHash == 0)
+	{
+		return nullptr;
+	}
+
+	FText3DFontFaceCache& FontFaceCache = CachedFontFaces.FindOrAdd(FontFaceHash, FText3DFontFaceCache(InFontFace));
+
+	/** Refresh font face with latest valid version */
+	FontFaceCache.UpdateFontFace(InFontFace);
+
+	return &FontFaceCache;
 }
 
 uint32 FText3DFontFaceCache::GetFontFaceHash(const TSharedPtr<FFreeTypeFace>& InFontFace)
@@ -157,7 +168,8 @@ uint32 FText3DFontFaceCache::GetFontFaceHash(const TSharedPtr<FFreeTypeFace>& In
 	{
 		const FString FaceFamily(Face->family_name);
 		const FString FaceStyle(Face->style_name);
-		return HashCombine(GetTypeHash(FaceFamily), GetTypeHash(FaceStyle));
+		const int32 FaceIndex = Face->face_index;
+		return HashCombine(HashCombine(GetTypeHash(FaceFamily), GetTypeHash(FaceStyle)), FaceIndex);
 	}
 
 	return 0;
@@ -168,6 +180,11 @@ FText3DFontFaceCache::FText3DFontFaceCache(const TSharedPtr<FFreeTypeFace>& InFo
 	, FontFaceHash(GetFontFaceHash(InFontFace))
 {
 }
+
+void FText3DFontFaceCache::UpdateFontFace(const TSharedPtr<FFreeTypeFace>& InFontFace)
+{
+	FontFace = InFontFace;
+}
 #endif
 
 uint32 FText3DFontFaceCache::GetGlyphMeshHash(uint32 InGlyphIndex, const FGlyphMeshParameters& InParameters)
@@ -295,7 +312,7 @@ void FText3DFontFaceCache::PrintCache() const
 	if (IsValid())
 	{
 		const FT_Face Face = FontFace->GetFace();
-		UE_LOG(LogText3D, Log, TEXT("Cached Font Face Family=%s Style=%s Hash=%u Meshes=%i Usage=%i"), *FString(Face->family_name), *FString(Face->style_name), FontFaceHash, GlyphMeshes.Num(), GetCacheRefCount());
+		UE_LOG(LogText3D, Log, TEXT("Cached Font Face Family=%s Style=%s Index=%i Hash=%u Meshes=%i Usage=%i"), *FString(Face->family_name), *FString(Face->style_name), Face->face_index, FontFaceHash, GlyphMeshes.Num(), GetCacheRefCount());
 	}
 #endif
 }
diff --git a/Engine/Plugins/Experimental/Text3D/Source/Text3D/Private/Subsystems/Text3DEngineSubsystem.h b/Engine/Plugins/Experimental/Text3D/Source/Text3D/Private/Subsystems/Text3DEngineSubsystem.h
--- a/Engine/Plugins/Experimental/Text3D/Source/Text3D/Private/Subsystems/Text3DEngineSubsystem.h
+++ b/Engine/Plugins/Experimental/Text3D/Source/Text3D/Private/Subsystems/Text3DEngineSubsystem.h
@@ -26,6 +26,7 @@ struct FText3DFontFaceCache
 
 #if WITH_FREETYPE
 	explicit FText3DFontFaceCache(const TSharedPtr<FFreeTypeFace>& InFontFace);
+	void UpdateFontFace(const TSharedPtr<FFreeTypeFace>& InFontFace);
 #endif
 
 	FText3DCachedMesh* FindOrBuildGlyphMesh(uint32 InGlyphIndex, const FGlyphMeshParameters& InParameters);
diff --git a/Engine/Plugins/Experimental/Text3D/Source/Text3D/Private/Templates/Text3DTemplateBase.cpp b/Engine/Plugins/Experimental/Text3D/Source/Text3D/Private/Templates/Text3DTemplateBase.cpp
new file mode 100644
--- /dev/null
+++ b/Engine/Plugins/Experimental/Text3D/Source/Text3D/Private/Templates/Text3DTemplateBase.cpp
@@ -0,0 +1,67 @@
+// Copyright Epic Games, Inc. All Rights Reserved.
+
+#include "Templates/Text3DTemplateBase.h"
+
+#include "Text3DComponent.h"
+
+#if WITH_EDITOR
+void UText3DTemplateBase::PostEditChangeProperty(FPropertyChangedEvent& InPropertyChangedEvent)
+{
+	Super::PostEditChangeProperty(InPropertyChangedEvent);
+
+	const FName MemberName = InPropertyChangedEvent.GetMemberPropertyName();
+
+	if (MemberName == GET_MEMBER_NAME_CHECKED(UText3DTemplateBase, TemplateName)
+		|| MemberName == GET_MEMBER_NAME_CHECKED(UText3DTemplateBase, Content))
+	{
+		OnTemplatePropertiesChanged();
+	}
+}
+#endif
+
+void UText3DTemplateBase::SetTemplateName(FName InName)
+{
+	if (InName == TemplateName)
+	{
+		return;
+	}
+
+	TemplateName = InName;
+	OnTemplatePropertiesChanged();
+}
+
+void UText3DTemplateBase::SetContent(const FText& InContent)
+{
+	if (InContent.EqualTo(Content))
+	{
+		return;
+	}
+
+	Content = InContent;
+	OnTemplatePropertiesChanged();
+}
+
+void UText3DTemplateBase::OnTemplatePropertiesChanged()
+{
+	if (!TemplateName.IsNone())
+	{
+		if (UText3DComponent* Text3DComponent = GetTypedOuter<UText3DComponent>())
+		{
+			Text3DComponent->RequestUpdate(EText3DRendererFlags::All, /** Immediate */false);
+		}
+	}
+}
+
+void UText3DTemplateBase::CollectTemplates(FFormatNamedArguments& InNamedArguments)
+{
+	if (!TemplateName.IsNone())
+	{
+		const FString TemplatePattern = TemplateName.ToString();
+		const FString EscapeTemplatePattern = FString::Printf(TEXT("\\%s\\"), *TemplatePattern);
+		const FString EscapeTemplateMarker = FString::Printf(TEXT("{%s}"), *TemplatePattern);
+
+		// Replaces {\template\} by {template} and {template} by its content
+		InNamedArguments.Add(*EscapeTemplatePattern, FText::FromString(EscapeTemplateMarker));
+		InNamedArguments.Add(*TemplatePattern, Content);
+	}
+}
diff --git a/Engine/Plugins/Experimental/Text3D/Source/Text3D/Private/Text3DComponent.cpp b/Engine/Plugins/Experimental/Text3D/Source/Text3D/Private/Text3DComponent.cpp
--- a/Engine/Plugins/Experimental/Text3D/Source/Text3D/Private/Text3DComponent.cpp
+++ b/Engine/Plugins/Experimental/Text3D/Source/Text3D/Private/Text3DComponent.cpp
@@ -12,6 +12,7 @@
 #include "Extensions/Text3DDefaultMaterialExtension.h"
 #include "Extensions/Text3DDefaultRenderingExtension.h"
 #include "Extensions/Text3DLayoutEffectBase.h"
+#include "Extensions/Text3DDefaultTemplateExtension.h"
 #include "Fonts/CompositeFont.h"
 #include "Internationalization/Regex.h"
 #include "Logs/Text3DLogs.h"
@@ -49,6 +50,7 @@ UText3DComponent::UText3DComponent()
 	TextRendererClass = UText3DStaticMeshesRenderer::StaticClass();
 
 	// Extensions
+	TemplateExtension = CreateDefaultSubobject<UText3DTemplateExtensionBase, UText3DDefaultTemplateExtension>(TEXT("TemplateExtension"));
 	FormatStyleExtension = CreateDefaultSubobject<UText3DFormatStyleExtensionBase, UText3DDefaultFormatStyleExtension>(TEXT("FormatStyleExtension"));
 	MaterialExtension = CreateDefaultSubobject<UText3DMaterialExtensionBase, UText3DDefaultMaterialExtension>(TEXT("MaterialExtension"));
 	GeometryExtension = CreateDefaultSubobject<UText3DGeometryExtensionBase, UText3DDefaultGeometryExtension>(TEXT("GeometryExtension"));
@@ -444,7 +446,12 @@ void UText3DComponent::ForceUpdateText()
 
 const FText& UText3DComponent::GetFormattedText() const
 {
-	return FormattedText.Get(Text);
+	if (const UText3DTemplateExtensionBase* Extension = GetTemplateExtension())
+	{
+		return Extension->GetFormattedText();
+	}
+
+	return FText::GetEmpty();
 }
 
 void UText3DComponent::SetText(const FText& Value)
@@ -615,6 +622,16 @@ UText3DRendererBase* UText3DComponent::GetTextRenderer() const
 	return TextRenderer;
 }
 
+UText3DTemplateExtensionBase* UText3DComponent::GetTemplateExtension(TSubclassOf<UText3DTemplateExtensionBase> InExtensionClass)
+{
+	if (!TemplateExtension || !InExtensionClass.Get())
+	{
+		return nullptr;
+	}
+
+	return TemplateExtension->IsA(InExtensionClass.Get()) ? TemplateExtension : nullptr;
+}
+
 UText3DFormatStyleExtensionBase* UText3DComponent::GetFormatStyleExtension(TSubclassOf<UText3DFormatStyleExtensionBase> InExtensionClass)
 {
 	if (!FormatStyleExtension || !InExtensionClass.Get())
@@ -1039,21 +1056,11 @@ void UText3DComponent::Rebuild(EText3DRendererFlags InUpdateFlags)
 
 	TextPreUpdateDelegate.Broadcast(this, InUpdateFlags);
 
-	TRACE_CPUPROFILER_EVENT_SCOPE(UText3DComponent::RebuildInternal);
-
-	bool bBroadcastEvent = false;
-
-	if (EnumHasAnyFlags(InUpdateFlags,EText3DRendererFlags::Geometry))
-	{
-		FormattedText = Text;
-		FormatText(FormattedText.GetValue());
-
-		bBroadcastEvent = true;
-	}
+	TRACE_CPUPROFILER_EVENT_SCOPE(UText3DComponent::Rebuild);
 
 	TextRenderer->Update(InUpdateFlags);
 
-	if (bBroadcastEvent)
+	if (EnumHasAnyFlags(InUpdateFlags,EText3DRendererFlags::Geometry))
 	{
 		UpdateStatistics();
 
diff --git a/Engine/Plugins/Experimental/Text3D/Source/Text3D/Public/Extensions/Text3DCharacterExtensionBase.h b/Engine/Plugins/Experimental/Text3D/Source/Text3D/Public/Extensions/Text3DCharacterExtensionBase.h
--- a/Engine/Plugins/Experimental/Text3D/Source/Text3D/Public/Extensions/Text3DCharacterExtensionBase.h
+++ b/Engine/Plugins/Experimental/Text3D/Source/Text3D/Public/Extensions/Text3DCharacterExtensionBase.h
@@ -15,10 +15,6 @@ class UText3DCharacterExtensionBase : public UText3DExtensionBase
 	GENERATED_BODY()
 
 public:
-	UText3DCharacterExtensionBase()
-		: UText3DExtensionBase(0)
-	{}
-
 	/** Get the amount of character */
 	virtual uint16 GetCharacterCount() const
 	{
diff --git a/Engine/Plugins/Experimental/Text3D/Source/Text3D/Public/Extensions/Text3DDefaultFormatStyleExtension.h b/Engine/Plugins/Experimental/Text3D/Source/Text3D/Public/Extensions/Text3DDefaultFormatStyleExtension.h
--- a/Engine/Plugins/Experimental/Text3D/Source/Text3D/Public/Extensions/Text3DDefaultFormatStyleExtension.h
+++ b/Engine/Plugins/Experimental/Text3D/Source/Text3D/Public/Extensions/Text3DDefaultFormatStyleExtension.h
@@ -32,7 +32,11 @@ protected:
     virtual EText3DExtensionResult PostRendererUpdate(const UE::Text3D::Renderer::FUpdateParameters& InParameters) override;
     //~ End UText3DExtensionBase
 
-	UPROPERTY(EditInstanceOnly, Instanced, Category="Style", meta=(AllowPrivateAccess="true"))
+	/**
+	 * Define your styles and use them in text like this <MyStyle>...</>,
+	 * the text between start and end tag will have this style applied
+	 */
+	UPROPERTY(EditInstanceOnly, Instanced, Category="Style", NoClear, meta=(AllowPrivateAccess="true"))
 	TSet<TObjectPtr<UText3DFormatStyleBase>> FormatStyles;
 
 private:
diff --git a/Engine/Plugins/Experimental/Text3D/Source/Text3D/Public/Extensions/Text3DDefaultTemplateExtension.h b/Engine/Plugins/Experimental/Text3D/Source/Text3D/Public/Extensions/Text3DDefaultTemplateExtension.h
new file mode 100644
--- /dev/null
+++ b/Engine/Plugins/Experimental/Text3D/Source/Text3D/Public/Extensions/Text3DDefaultTemplateExtension.h
@@ -0,0 +1,41 @@
+// Copyright Epic Games, Inc. All Rights Reserved.
+
+#pragma once
+
+#include "Containers/Set.h"
+#include "Text3DTemplateExtensionBase.h"
+#include "Text3DDefaultTemplateExtension.generated.h"
+
+class UText3DTemplateBase;
+
+UCLASS(MinimalAPI)
+class UText3DDefaultTemplateExtension : public UText3DTemplateExtensionBase
+{
+	GENERATED_BODY()
+
+protected:
+	//~ Begin UObject
+#if WITH_EDITOR
+	virtual void PostEditChangeProperty(FPropertyChangedEvent& InPropertyChangedEvent) override;
+#endif
+	//~ End UObject
+
+	//~ Begin UText3DTemplateExtensionBase
+	virtual const FText& GetFormattedText() const override;
+	//~ End UText3DTemplateExtensionBase
+
+	//~ Begin UText3DExtensionBase
+	virtual EText3DExtensionResult PreRendererUpdate(const UE::Text3D::Renderer::FUpdateParameters& InParameters) override;
+	virtual EText3DExtensionResult PostRendererUpdate(const UE::Text3D::Renderer::FUpdateParameters& InParameters) override;
+	//~ End UText3DExtensionBase
+
+	/**
+	 * Define your templates and use them in text like this {MyTemplate},
+	 * {MyTemplate} will be swapped by the content of the template called MyTemplate
+	 */
+	UPROPERTY(EditInstanceOnly, Instanced, Category="Template", NoClear, meta=(AllowPrivateAccess="true"))
+	TSet<TObjectPtr<UText3DTemplateBase>> Templates;
+
+	UPROPERTY(Transient)
+	FText FormattedText;
+};
diff --git a/Engine/Plugins/Experimental/Text3D/Source/Text3D/Public/Extensions/Text3DEffectExtensionBase.h b/Engine/Plugins/Experimental/Text3D/Source/Text3D/Public/Extensions/Text3DEffectExtensionBase.h
--- a/Engine/Plugins/Experimental/Text3D/Source/Text3D/Public/Extensions/Text3DEffectExtensionBase.h
+++ b/Engine/Plugins/Experimental/Text3D/Source/Text3D/Public/Extensions/Text3DEffectExtensionBase.h
@@ -52,7 +52,7 @@ class UText3DEffectExtensionBase : public UText3DExtensionBase
 
 public:
 	UText3DEffectExtensionBase()
-		: UText3DExtensionBase(/** Priority */0)
+		: UText3DExtensionBase()
 	{}
 
 	explicit UText3DEffectExtensionBase(uint16 InPriority)
diff --git a/Engine/Plugins/Experimental/Text3D/Source/Text3D/Public/Extensions/Text3DFormatStyleExtensionBase.h b/Engine/Plugins/Experimental/Text3D/Source/Text3D/Public/Extensions/Text3DFormatStyleExtensionBase.h
--- a/Engine/Plugins/Experimental/Text3D/Source/Text3D/Public/Extensions/Text3DFormatStyleExtensionBase.h
+++ b/Engine/Plugins/Experimental/Text3D/Source/Text3D/Public/Extensions/Text3DFormatStyleExtensionBase.h
@@ -19,7 +19,7 @@ class UText3DFormatStyleExtensionBase : public UText3DExtensionBase
 
 public:
 	UText3DFormatStyleExtensionBase()
-		: UText3DExtensionBase(/** Priority */1)
+		: UText3DExtensionBase(UE::Text3D::Priority::Style)
 	{}
 
 	/** Default style used by text when no other style is applied */
diff --git a/Engine/Plugins/Experimental/Text3D/Source/Text3D/Public/Extensions/Text3DGeometryExtensionBase.h b/Engine/Plugins/Experimental/Text3D/Source/Text3D/Public/Extensions/Text3DGeometryExtensionBase.h
--- a/Engine/Plugins/Experimental/Text3D/Source/Text3D/Public/Extensions/Text3DGeometryExtensionBase.h
+++ b/Engine/Plugins/Experimental/Text3D/Source/Text3D/Public/Extensions/Text3DGeometryExtensionBase.h
@@ -20,7 +20,7 @@ class UText3DGeometryExtensionBase : public UText3DExtensionBase
 
 public:
 	UText3DGeometryExtensionBase()
-		: UText3DExtensionBase(/** Priority */1)
+		: UText3DExtensionBase(UE::Text3D::Priority::Geometry)
 	{}
 
 	// Deprecated do not use anymore, use character getter directly
diff --git a/Engine/Plugins/Experimental/Text3D/Source/Text3D/Public/Extensions/Text3DLayoutEffectBase.h b/Engine/Plugins/Experimental/Text3D/Source/Text3D/Public/Extensions/Text3DLayoutEffectBase.h
--- a/Engine/Plugins/Experimental/Text3D/Source/Text3D/Public/Extensions/Text3DLayoutEffectBase.h
+++ b/Engine/Plugins/Experimental/Text3D/Source/Text3D/Public/Extensions/Text3DLayoutEffectBase.h
@@ -13,6 +13,6 @@ class UText3DLayoutEffectBase : public UText3DEffectExtensionBase
 
 public:
 	UText3DLayoutEffectBase()
-		: UText3DEffectExtensionBase(/** Priority */3)
+		: UText3DEffectExtensionBase(UE::Text3D::Priority::Effect)
 	{}
 };
diff --git a/Engine/Plugins/Experimental/Text3D/Source/Text3D/Public/Extensions/Text3DLayoutExtensionBase.h b/Engine/Plugins/Experimental/Text3D/Source/Text3D/Public/Extensions/Text3DLayoutExtensionBase.h
--- a/Engine/Plugins/Experimental/Text3D/Source/Text3D/Public/Extensions/Text3DLayoutExtensionBase.h
+++ b/Engine/Plugins/Experimental/Text3D/Source/Text3D/Public/Extensions/Text3DLayoutExtensionBase.h
@@ -16,7 +16,7 @@ class UText3DLayoutExtensionBase : public UText3DExtensionBase
 
 public:
 	UText3DLayoutExtensionBase()
-		: UText3DExtensionBase(/** Priority */2)
+		: UText3DExtensionBase(UE::Text3D::Priority::Layout)
 	{}
 
 	virtual float GetTextHeight() const
diff --git a/Engine/Plugins/Experimental/Text3D/Source/Text3D/Public/Extensions/Text3DMaterialExtensionBase.h b/Engine/Plugins/Experimental/Text3D/Source/Text3D/Public/Extensions/Text3DMaterialExtensionBase.h
--- a/Engine/Plugins/Experimental/Text3D/Source/Text3D/Public/Extensions/Text3DMaterialExtensionBase.h
+++ b/Engine/Plugins/Experimental/Text3D/Source/Text3D/Public/Extensions/Text3DMaterialExtensionBase.h
@@ -16,7 +16,7 @@ class UText3DMaterialExtensionBase : public UText3DExtensionBase
 
 public:
 	UText3DMaterialExtensionBase()
-		: UText3DExtensionBase(/** Priority */3)
+		: UText3DExtensionBase(UE::Text3D::Priority::Material)
 	{}
 
 	/** Set the material for a specific group */
diff --git a/Engine/Plugins/Experimental/Text3D/Source/Text3D/Public/Extensions/Text3DRenderingExtensionBase.h b/Engine/Plugins/Experimental/Text3D/Source/Text3D/Public/Extensions/Text3DRenderingExtensionBase.h
--- a/Engine/Plugins/Experimental/Text3D/Source/Text3D/Public/Extensions/Text3DRenderingExtensionBase.h
+++ b/Engine/Plugins/Experimental/Text3D/Source/Text3D/Public/Extensions/Text3DRenderingExtensionBase.h
@@ -14,7 +14,7 @@ class UText3DRenderingExtensionBase : public UText3DExtensionBase
 
 public:
 	UText3DRenderingExtensionBase()
-		: UText3DExtensionBase(/** Priority */4)
+		: UText3DExtensionBase(UE::Text3D::Priority::Rendering)
 	{}
 
 	virtual bool GetTextCastShadow() const
diff --git a/Engine/Plugins/Experimental/Text3D/Source/Text3D/Public/Extensions/Text3DTemplateExtensionBase.h b/Engine/Plugins/Experimental/Text3D/Source/Text3D/Public/Extensions/Text3DTemplateExtensionBase.h
new file mode 100644
--- /dev/null
+++ b/Engine/Plugins/Experimental/Text3D/Source/Text3D/Public/Extensions/Text3DTemplateExtensionBase.h
@@ -0,0 +1,25 @@
+// Copyright Epic Games, Inc. All Rights Reserved.
+
+#pragma once
+
+#include "Text3DExtensionBase.h"
+#include "Text3DTypes.h"
+#include "Text3DTemplateExtensionBase.generated.h"
+
+/** Extension that handles text transformation */
+UCLASS(MinimalAPI, Abstract)
+class UText3DTemplateExtensionBase : public UText3DExtensionBase
+{
+	GENERATED_BODY()
+
+public:
+	UText3DTemplateExtensionBase()
+		: UText3DExtensionBase(UE::Text3D::Priority::Template)
+	{}
+
+	/** Get the transformed text */
+	virtual const FText& GetFormattedText() const
+	{
+		return FText::GetEmpty();
+	}
+};
diff --git a/Engine/Plugins/Experimental/Text3D/Source/Text3D/Public/Templates/Text3DTemplateBase.h b/Engine/Plugins/Experimental/Text3D/Source/Text3D/Public/Templates/Text3DTemplateBase.h
new file mode 100644
--- /dev/null
+++ b/Engine/Plugins/Experimental/Text3D/Source/Text3D/Public/Templates/Text3DTemplateBase.h
@@ -0,0 +1,46 @@
+// Copyright Epic Games, Inc. All Rights Reserved.
+
+#pragma once
+
+#include "UObject/Object.h"
+#include "Text3DTemplateBase.generated.h"
+
+/** Text3D template to swap variable by its content */
+UCLASS(MinimalAPI, EditInlineNew, ClassGroup=Text3D, DisplayName="Template", AutoExpandCategories=(Template))
+class UText3DTemplateBase : public UObject
+{
+	GENERATED_BODY()
+
+public:
+	void SetTemplateName(FName InName);
+	FName GetTemplateName() const
+	{
+		return TemplateName;
+	}
+
+	void SetContent(const FText& InContent);
+	const FText& GetContent() const
+	{
+		return Content;
+	}
+
+	/** Collects named templates and their content for replacement */
+	virtual void CollectTemplates(FFormatNamedArguments& InNamedArguments);
+
+protected:
+	//~ Begin UObject
+#if WITH_EDITOR
+	virtual void PostEditChangeProperty(FPropertyChangedEvent& InPropertyChangedEvent) override;
+#endif
+	//~ End UObject
+
+	void OnTemplatePropertiesChanged();
+
+	/** Template variable that will be swapped */
+	UPROPERTY(EditAnywhere, Setter, Getter, Category="Template", meta=(AllowPrivateAccess="true"))
+	FName TemplateName = NAME_None;
+
+	/** Template content that replaces variable */
+	UPROPERTY(EditAnywhere, Setter, Getter, Category="Template", meta=(AllowPrivateAccess="true"))
+	FText Content;
+};
diff --git a/Engine/Plugins/Experimental/Text3D/Source/Text3D/Public/Text3DComponent.h b/Engine/Plugins/Experimental/Text3D/Source/Text3D/Public/Text3DComponent.h
--- a/Engine/Plugins/Experimental/Text3D/Source/Text3D/Public/Text3DComponent.h
+++ b/Engine/Plugins/Experimental/Text3D/Source/Text3D/Public/Text3DComponent.h
@@ -18,6 +18,7 @@ class UText3DGeometryExtensionBase;
 class UText3DLayoutEffectBase;
 class UText3DLayoutExtensionBase;
 class UText3DMaterialExtensionBase;
+class UText3DTemplateExtensionBase;
 class UText3DRendererBase;
 class UText3DRenderingExtensionBase;
 struct FTypefaceEntry;
@@ -29,6 +30,9 @@ UCLASS(MinimalAPI
 	, meta = (BlueprintSpawnableComponent))
 class UText3DComponent : public USceneComponent
 {
+	// Needed to support older API (see FormatText)
+	friend class UText3DDefaultTemplateExtension;
+
 	GENERATED_BODY()
 
 public:
@@ -268,6 +272,20 @@ public:
 	
 	TEXT3D_API UText3DRendererBase* GetTextRenderer() const;
 
+	UFUNCTION(BlueprintCallable, Category="Text3D", meta=(DeterminesOutputType="InExtensionClass"))
+	UText3DTemplateExtensionBase* GetTemplateExtension(TSubclassOf<UText3DTemplateExtensionBase> InExtensionClass);
+	
+	UText3DTemplateExtensionBase* GetTemplateExtension() const
+	{
+		return TemplateExtension;
+	}
+	
+	template <typename T UE_REQUIRES(TIsDerivedFrom<T, UText3DTemplateExtensionBase>::Value)>
+	T* GetCastedTemplateExtension() const
+	{
+		return Cast<T>(TemplateExtension);
+	}
+
 	UFUNCTION(BlueprintCallable, Category="Text3D", meta=(DeterminesOutputType="InExtensionClass"))
 	UText3DFormatStyleExtensionBase* GetFormatStyleExtension(TSubclassOf<UText3DFormatStyleExtensionBase> InExtensionClass);
 	
@@ -445,10 +463,6 @@ protected:
 	UPROPERTY(EditAnywhere, Getter, Setter, Category = "Text", meta = (MultiLine = true, AllowPrivateAccess = "true"))
 	FText Text;
 
-	/** Original text formatted */
-	UPROPERTY(Transient)
-	TOptional<FText> FormattedText;
-
 	/** Text font defines the style of rendered characters */
 	UPROPERTY(EditAnywhere, Getter, Setter, DisplayName = "Default Font", Category = "Text", meta = (AdvancedFontPicker, AllowPrivateAccess = "true"))
 	TObjectPtr<UFont> Font;
@@ -461,6 +475,9 @@ protected:
 	UPROPERTY(EditAnywhere, Getter = "GetEnforceUpperCase", Setter = "SetEnforceUpperCase", Category = "Text", meta = (AllowPrivateAccess="true"))
 	bool bEnforceUpperCase = false;
 
+	UPROPERTY(VisibleInstanceOnly, Instanced, Category="Template", meta=(AllowPrivateAccess="true"))
+	TObjectPtr<UText3DTemplateExtensionBase> TemplateExtension;
+
 	UPROPERTY(VisibleInstanceOnly, Instanced, Category="Style", meta=(AllowPrivateAccess="true"))
 	TObjectPtr<UText3DFormatStyleExtensionBase> FormatStyleExtension;
 
diff --git a/Engine/Plugins/Experimental/Text3D/Source/Text3D/Public/Text3DTypes.h b/Engine/Plugins/Experimental/Text3D/Source/Text3D/Public/Text3DTypes.h
--- a/Engine/Plugins/Experimental/Text3D/Source/Text3D/Public/Text3DTypes.h
+++ b/Engine/Plugins/Experimental/Text3D/Source/Text3D/Public/Text3DTypes.h
@@ -292,6 +292,17 @@ struct FText3DCachedMesh
 
 namespace UE::Text3D
 {
+	namespace Priority
+	{
+		static constexpr uint16 Template = 0;
+		static constexpr uint16 Geometry = 1;
+		static constexpr uint16 Style = 2;
+		static constexpr uint16 Layout = 3;
+		static constexpr uint16 Effect = 4;
+		static constexpr uint16 Material = 5;
+		static constexpr uint16 Rendering = 6;
+	}
+
 	namespace Metrics
 	{
 		/** Original value to match size across different Text3D version */
diff --git a/Engine/Plugins/Experimental/Text3D/Source/Text3DEditor/Private/Customizations/Text3DEditorTextComponentDetailCustomization.cpp b/Engine/Plugins/Experimental/Text3D/Source/Text3DEditor/Private/Customizations/Text3DEditorTextComponentDetailCustomization.cpp
--- a/Engine/Plugins/Experimental/Text3D/Source/Text3DEditor/Private/Customizations/Text3DEditorTextComponentDetailCustomization.cpp
+++ b/Engine/Plugins/Experimental/Text3D/Source/Text3DEditor/Private/Customizations/Text3DEditorTextComponentDetailCustomization.cpp
@@ -22,6 +22,7 @@ namespace UE::Text3DEditor::Customization
 		const TArray<FName> CategoryNames =
 		{
 			TEXT("Text"),
+			TEXT("Template"),
 			TEXT("Style"),
 			TEXT("Character"),
 			TEXT("Layout"),
@@ -61,6 +62,7 @@ namespace UE::Text3DEditor::Customization
 			const FName TextSectionName = TEXT("Text");
 			const TSharedRef<FPropertySection> TextSection = PropertyModule.FindOrCreateSection(ComponentClassName, TextSectionName, FText::FromName(TextSectionName));
 			TextSection->AddCategory(TEXT("Text"));
+			TextSection->AddCategory(TEXT("Template"));
 
 			const FName StyleSectionName = TEXT("Style");
 			const TSharedRef<FPropertySection> MaterialSection = PropertyModule.FindOrCreateSection(ComponentClassName, StyleSectionName, FText::FromName(StyleSectionName));
diff --git a/Engine/Plugins/VirtualProduction/Avalanche/Source/AvalancheTextEditor/Private/Visualizer/AvaTextVisualizer.cpp b/Engine/Plugins/VirtualProduction/Avalanche/Source/AvalancheTextEditor/Private/Visualizer/AvaTextVisualizer.cpp
--- a/Engine/Plugins/VirtualProduction/Avalanche/Source/AvalancheTextEditor/Private/Visualizer/AvaTextVisualizer.cpp
+++ b/Engine/Plugins/VirtualProduction/Avalanche/Source/AvalancheTextEditor/Private/Visualizer/AvaTextVisualizer.cpp
@@ -1001,7 +1001,7 @@ void FAvaTextVisualizer::StoreTextMetrics(const UText3DComponent* InText3DComp)
 {
 	Bounds = GetComponentBounds(InText3DComp);
 	TArray<FString> Lines;
-	InText3DComp->GetText().ToString().ParseIntoArrayLines(Lines);
+	InText3DComp->GetFormattedText().ToString().ParseIntoArrayLines(Lines);
 	LineCount = Lines.Num();
 	LineHeight = (Bounds.GetSize().Z - (InText3DComp->GetLineSpacing() * (static_cast<float>(LineCount) - 1.f))) / static_cast<float>(LineCount);
 }
