MotionDesign : Text3D
- Added color style with additional materials
- Added ability to define a style set asset that can be reused (with events once changed to update instances using that asset), Default style set can be defined in project settings
- Added style rule to inherit or override (font, color)
- Added button to open project settings window in editor
- Renamed new FormatStyleExtension to StyleExtension
- Renamed new TemplateExtension to TokenExtension

#DisabledValidations NonUnityAndStaticAnalysis

diff --git a/Engine/Plugins/Experimental/Text3D/Source/Text3D/Private/Characters/Text3DCharacterBase.cpp b/Engine/Plugins/Experimental/Text3D/Source/Text3D/Private/Characters/Text3DCharacterBase.cpp
--- a/Engine/Plugins/Experimental/Text3D/Source/Text3D/Private/Characters/Text3DCharacterBase.cpp
+++ b/Engine/Plugins/Experimental/Text3D/Source/Text3D/Private/Characters/Text3DCharacterBase.cpp
@@ -44,6 +44,11 @@ void UText3DCharacterBase::SetFontFaceGlyphHandle(const UE::Text3D::Geometry::FC
 	FontFaceGlyphHandle = InHandle;
 }
 
+void UText3DCharacterBase::SetStyleTag(FName InStyle)
+{
+	StyleTag = InStyle;
+}
+
 uint32 UText3DCharacterBase::GetGlyphIndex() const
 {
 	return FontFaceGlyphHandle.GetGlyphIndex();
diff --git a/Engine/Plugins/Experimental/Text3D/Source/Text3D/Private/Extensions/Text3DDefaultCharacterExtension.cpp b/Engine/Plugins/Experimental/Text3D/Source/Text3D/Private/Extensions/Text3DDefaultCharacterExtension.cpp
--- a/Engine/Plugins/Experimental/Text3D/Source/Text3D/Private/Extensions/Text3DDefaultCharacterExtension.cpp
+++ b/Engine/Plugins/Experimental/Text3D/Source/Text3D/Private/Extensions/Text3DDefaultCharacterExtension.cpp
@@ -59,15 +59,20 @@ void UText3DDefaultCharacterExtension::AllocateTextCharacters(uint16 InCharacter
 	{
 		for (int32 CharacterIndex = 0; CharacterIndex < FMath::Abs(RemainingCharacterCount); ++CharacterIndex)
 		{
+			UText3DCharacterBase* TextCharacter = nullptr;
+
 			if (!TextCharactersPool.IsEmpty())
 			{
-				TextCharacters.Add(TextCharactersPool.Pop());
+				TextCharacter = TextCharactersPool.Pop();
 			}
-			else
+
+			if (!TextCharacter)
 			{
 				const FName ObjectName = MakeUniqueObjectName(this, UText3DDefaultCharacter::StaticClass(), FName(TEXT("Char")));
-				TextCharacters.Add(NewObject<UText3DCharacterBase>(this, UText3DDefaultCharacter::StaticClass(), ObjectName, RF_NoFlags));
+				TextCharacter = NewObject<UText3DCharacterBase>(this, UText3DDefaultCharacter::StaticClass(), ObjectName, RF_NoFlags);
 			}
+
+			TextCharacters.Add(TextCharacter);
 		}
 	}
 
diff --git a/Engine/Plugins/Experimental/Text3D/Source/Text3D/Private/Extensions/Text3DDefaultFormatStyleExtension.cpp b/Engine/Plugins/Experimental/Text3D/Source/Text3D/Private/Extensions/Text3DDefaultFormatStyleExtension.cpp
deleted file mode 100644
diff --git a/Engine/Plugins/Experimental/Text3D/Source/Text3D/Private/Extensions/Text3DDefaultLayoutExtension.cpp b/Engine/Plugins/Experimental/Text3D/Source/Text3D/Private/Extensions/Text3DDefaultLayoutExtension.cpp
--- a/Engine/Plugins/Experimental/Text3D/Source/Text3D/Private/Extensions/Text3DDefaultLayoutExtension.cpp
+++ b/Engine/Plugins/Experimental/Text3D/Source/Text3D/Private/Extensions/Text3DDefaultLayoutExtension.cpp
@@ -7,8 +7,8 @@
 #include "Characters/Text3DCharacterBase.h"
 #include "Engine/Font.h"
 #include "Extensions/Text3DCharacterExtensionBase.h"
-#include "Extensions/Text3DFormatStyleExtensionBase.h"
 #include "Extensions/Text3DGeometryExtensionBase.h"
+#include "Extensions/Text3DStyleExtensionBase.h"
 #include "Framework/Text/RichTextLayoutMarshaller.h"
 #include "Framework/Text/RichTextMarkupProcessing.h"
 #include "LayoutBuilders/Text3DLayout.h"
@@ -175,9 +175,11 @@ void UText3DDefaultLayoutExtension::SetScaleProportionally(const bool Value)
 
 float UText3DDefaultLayoutExtension::GetTextHeight() const
 {
-	float Height = FMath::Max(ShapedText->Lines.Num() - 1, 0) * LineSpacing;
+	using namespace UE::Text3D::Layout;
 
-	for (const FShapedGlyphLine& Line : ShapedText->Lines)
+	float Height = FMath::Max(GlyphText->Lines.Num() - 1, 0) * LineSpacing;
+
+	for (const FGlyphLine& Line : GlyphText->Lines)
 	{
 		Height += Line.MaxFontHeight;
 	}
@@ -192,6 +194,8 @@ FVector UText3DDefaultLayoutExtension::GetTextScale() const
 
 EText3DExtensionResult UText3DDefaultLayoutExtension::PreRendererUpdate(const UE::Text3D::Renderer::FUpdateParameters& InParameters)
 {
+	using namespace UE::Text3D::Layout;
+
 	if (InParameters.CurrentFlag != EText3DRendererFlags::Geometry && InParameters.CurrentFlag != EText3DRendererFlags::Layout)
 	{
 		return EText3DExtensionResult::Active;
@@ -199,9 +203,9 @@ EText3DExtensionResult UText3DDefaultLayoutExtension::PreRendererUpdate(const UE
 
 	const UText3DComponent* Text3DComponent = GetText3DComponent();
 	
-	if (!ShapedText.IsValid())
+	if (!GlyphText.IsValid())
 	{
-		ShapedText = MakeShared<FText3DShapedGlyphText>();
+		GlyphText = MakeShared<FGlyphText>();
 	}
 
 	if (!TextLayout.IsValid())
@@ -215,16 +219,16 @@ EText3DExtensionResult UText3DDefaultLayoutExtension::PreRendererUpdate(const UE
 		TextLayoutMarshaller = FRichTextLayoutMarshaller::Create(Parser, nullptr, TArray<TSharedRef<ITextDecorator>>(), nullptr);
 	}
 
-	ShapedText->Reset();
+	GlyphText->Reset();
 
-	const UText3DFormatStyleExtensionBase* FormatStyleExtension = Text3DComponent->GetFormatStyleExtension();
-	FText3DLayoutShaper::Get()->ShapeBidirectionalText(FormatStyleExtension->GetCustomStyles(), *FormatStyleExtension->GetDefaultStyle(), Text3DComponent->GetFormattedText().ToString(), TextLayout, TextLayoutMarshaller, ShapedText->Lines);
+	const UText3DStyleExtensionBase* FormatStyleExtension = Text3DComponent->GetStyleExtension();
+	FText3DLayoutShaper::Get()->ShapeBidirectionalText(FormatStyleExtension->GetCustomStyles(), *FormatStyleExtension->GetDefaultStyle(), Text3DComponent->GetFormattedText().ToString(), TextLayout, TextLayoutMarshaller, GlyphText->Lines);
 
-	const int32 CharacterCount = Algo::TransformAccumulate(ShapedText->Lines, [](const FShapedGlyphLine& InLine)
+	const int32 CharacterCount = Algo::TransformAccumulate(GlyphText->Lines, [](const FGlyphLine& InLine)
 	{
-		return Algo::CountIf(InLine.GlyphsToRender, [](const FShapedGlyphEntry& InGlyph)
+		return Algo::CountIf(InLine.Glyphs, [](const FGlyphEntry& InGlyph)
 		{
-			return FText3DShapedGlyphText::IsValidGlyph(InGlyph);
+			return FGlyphText::IsValidGlyph(InGlyph);
 		});
 	},0);
 
@@ -234,27 +238,34 @@ EText3DExtensionResult UText3DDefaultLayoutExtension::PreRendererUpdate(const UE
 		CharacterExtension->AllocateCharacters(CharacterCount);
 	}
 
-	ShapedText->Kernings.SetNum(CharacterCount);
-	ShapedText->FontFaces.SetNum(CharacterCount);
+	GlyphText->Kernings.SetNum(CharacterCount);
+	GlyphText->FontFaces.SetNum(CharacterCount);
 
 	const UText3DGeometryExtensionBase* GeometryExtension = Text3DComponent->GetGeometryExtension();
 	uint16 CharacterIndex = 0;
 
 #if WITH_FREETYPE
-	for (const TConstEnumerateRef<FShapedGlyphLine> ShapedLine : EnumerateRange(ShapedText->Lines))
+	for (const TConstEnumerateRef<FGlyphLine> GlyphLine : EnumerateRange(GlyphText->Lines))
 	{
-		const bool bIsRTL = (ShapedLine->TextDirection == TextBiDi::ETextDirection::RightToLeft);
+		const bool bIsRTL = (GlyphLine->TextDirection == TextBiDi::ETextDirection::RightToLeft);
 
-		for (const TConstEnumerateRef<FShapedGlyphEntry> GlyphEntry : EnumerateRange(ShapedLine->GlyphsToRender))
+		for (const TConstEnumerateRef<FGlyphEntry> GlyphEntry : EnumerateRange(GlyphLine->Glyphs))
 		{
-			if (FText3DShapedGlyphText::IsValidGlyph(*GlyphEntry))
+			if (FGlyphText::IsValidGlyph(*GlyphEntry))
 			{
-				const TSharedPtr<FFreeTypeFace> FontFacePtr = GlyphEntry->FontFaceData->FontFace.Pin();
+				const TSharedPtr<FFreeTypeFace> FontFacePtr = GlyphEntry->Entry.FontFaceData->FontFace.Pin();
 				const uint16 Index = bIsRTL ? CharacterCount - 1 - CharacterIndex : CharacterIndex;
-				const UText3DCharacterBase* Character = Text3DComponent->GetCharacter(Index);
-				ShapedText->Kernings[CharacterIndex] = Character->GetCharacterKerning();
-				ShapedText->FontFaces[CharacterIndex] = FontFacePtr;
-				CharacterIndex++;
+				if (const UText3DCharacterBase* Character = Text3DComponent->GetCharacter(Index))
+				{
+					GlyphText->Kernings[CharacterIndex] = Character->GetCharacterKerning();
+                    GlyphText->FontFaces[CharacterIndex] = FontFacePtr;
+                    CharacterIndex++;
+				}
+				else
+				{
+					UE_LOG(LogText3D, Warning, TEXT("Invalid character object returned for index %i, cannot proceed"), CharacterIndex)
+					return EText3DExtensionResult::Failed;
+				}
 			}
 		}
 	}
@@ -263,23 +274,23 @@ EText3DExtensionResult UText3DDefaultLayoutExtension::PreRendererUpdate(const UE
 	return EText3DExtensionResult::Failed;
 #endif
 
-	ShapedText->Tracking = Tracking;
-	ShapedText->WordSpacing = WordSpacing;
-	ShapedText->MaxWidth = MaxWidth;
-	ShapedText->bWrap = bUseMaxWidth && MaxWidthBehavior == EText3DMaxWidthHandling::WrapAndScale;
-	ShapedText->CalculateWidth();
+	GlyphText->Tracking = Tracking;
+	GlyphText->WordSpacing = WordSpacing;
+	GlyphText->MaxWidth = MaxWidth;
+	GlyphText->bWrap = bUseMaxWidth && MaxWidthBehavior == EText3DMaxWidthHandling::WrapAndScale;
+	GlyphText->CalculateWidth();
 	CalculateTextScale();
 
 	CharacterIndex = 0;
 	UText3DEngineSubsystem* TextSubsystem = UText3DEngineSubsystem::Get();
-	for (TConstEnumerateRef<FShapedGlyphLine> ShapedLine : EnumerateRange(ShapedText->Lines))
+	for (TConstEnumerateRef<FGlyphLine> GlyphLine : EnumerateRange(GlyphText->Lines))
 	{
-		FVector Location = GetLineLocation(ShapedLine.GetIndex());
-		const bool bIsRTL = (ShapedLine->TextDirection == TextBiDi::ETextDirection::RightToLeft);
+		FVector Location = GetLineLocation(GlyphLine.GetIndex());
+		const bool bIsRTL = (GlyphLine->TextDirection == TextBiDi::ETextDirection::RightToLeft);
 
-		for (TConstEnumerateRef<FShapedGlyphEntry> ShapedGlyph : EnumerateRange(ShapedLine->GlyphsToRender))
+		for (TConstEnumerateRef<FGlyphEntry> GlyphEntry : EnumerateRange(GlyphLine->Glyphs))
 		{
-			if (FText3DShapedGlyphText::IsValidGlyph(*ShapedGlyph))
+			if (FGlyphText::IsValidGlyph(*GlyphEntry))
 			{
 				const uint16 Index = bIsRTL ? CharacterCount - 1 - CharacterIndex : CharacterIndex;
 				UText3DCharacterBase* Character = Text3DComponent->GetCharacter(Index);
@@ -287,8 +298,8 @@ EText3DExtensionResult UText3DDefaultLayoutExtension::PreRendererUpdate(const UE
 				if (EnumHasAnyFlags(InParameters.UpdateFlags, EText3DRendererFlags::Geometry))
 				{
 #if WITH_FREETYPE
-					TSharedPtr<FFreeTypeFace> CharTypeface = ShapedText->FontFaces[CharacterIndex];
-					uint32 FontFaceHash = FText3DFontFaceCache::GetFontFaceHash(CharTypeface);
+					TSharedPtr<FFreeTypeFace> CharTypeface = GlyphText->FontFaces[CharacterIndex];
+					const uint32 FontFaceHash = FText3DFontFaceCache::GetFontFaceHash(CharTypeface);
 					// Cache font face to keep it alive and use inner FT_Face to build glyphs
 					FText3DFontFaceCache* FontFaceCache = TextSubsystem->FindOrAddCachedFontFace(CharTypeface);
 #else
@@ -299,12 +310,13 @@ EText3DExtensionResult UText3DDefaultLayoutExtension::PreRendererUpdate(const UE
 					if (FontFaceCache)
 					{
 						// Load glyph here to get proper mesh metrics
-						UE::Text3D::Geometry::FCachedFontFaceGlyphHandle GlyphHandle(FontFaceCache, ShapedGlyph->GlyphIndex, *GeometryExtension->GetGlyphMeshParameters());
+						UE::Text3D::Geometry::FCachedFontFaceGlyphHandle GlyphHandle(FontFaceCache, GlyphEntry->Entry.GlyphIndex, *GeometryExtension->GetGlyphMeshParameters());
 						Character->SetFontFaceGlyphHandle(GlyphHandle);
+						Character->SetStyleTag(GlyphEntry->StyleTag);
 					}
 					else
 					{
-						UE_LOG(LogText3D, Error, TEXT("Failed to get cached font data for '%u %i' in Text3D geometry extension"), FontFaceHash, ShapedGlyph->GlyphIndex);
+						UE_LOG(LogText3D, Error, TEXT("Failed to get cached font data for '%u %i' in Text3D geometry extension"), FontFaceHash, GlyphEntry->Entry.GlyphIndex);
 					}
 				}
 
@@ -331,16 +343,16 @@ EText3DExtensionResult UText3DDefaultLayoutExtension::PreRendererUpdate(const UE
 					break;
 				}
 
-				if (ShapedGlyph.GetIndex() != 0)
+				if (GlyphEntry.GetIndex() != 0)
 				{
-					Shift.Y += Tracking + ShapedText->Kernings[CharacterIndex];
+					Shift.Y += Tracking + GlyphText->Kernings[CharacterIndex];
 				}
 
 				CharacterTransform.Accumulate(FTransform(Location + Shift));
 				CharacterIndex++;
 			}
 
-			Location.Y += ShapedLine->GlyphAdvances[ShapedGlyph.GetIndex()];
+			Location.Y += GlyphLine->GlyphAdvances[GlyphEntry.GetIndex()];
 		}
 	}
 
@@ -354,10 +366,12 @@ EText3DExtensionResult UText3DDefaultLayoutExtension::PostRendererUpdate(const U
 
 void UText3DDefaultLayoutExtension::CalculateTextScale()
 {
+	using namespace UE::Text3D::Layout;
+
 	FVector Scale(1.0f, 1.0f, 1.0f);
 
 	float TextMaxWidth = 0.0f;
-	for (const FShapedGlyphLine& ShapedLine : ShapedText->Lines)
+	for (const FGlyphLine& ShapedLine : GlyphText->Lines)
 	{
 		TextMaxWidth = FMath::Max(TextMaxWidth, ShapedLine.Width);
 	}
@@ -391,15 +405,17 @@ void UText3DDefaultLayoutExtension::CalculateTextScale()
 
 FVector UText3DDefaultLayoutExtension::GetLineLocation(int32 LineIndex) const
 {
+	using namespace UE::Text3D::Layout;
+
 	float HorizontalOffset = 0.0f;
 	float VerticalOffset = 0.0f;
 
-	if (LineIndex < 0 || LineIndex >= ShapedText->Lines.Num())
+	if (LineIndex < 0 || LineIndex >= GlyphText->Lines.Num())
 	{
 		return FVector::ZeroVector;
 	}
 
-	const FShapedGlyphLine& ShapedLine = ShapedText->Lines[LineIndex];
+	const FGlyphLine& ShapedLine = GlyphText->Lines[LineIndex];
 
 	if (HorizontalAlignment == EText3DHorizontalTextAlignment::Center)
 	{
@@ -429,7 +445,7 @@ FVector UText3DDefaultLayoutExtension::GetLineLocation(int32 LineIndex) const
 	// Offset line based on cumulative height above it
 	for (int32 Index = 0; Index < LineIndex; ++Index)
 	{
-		const FShapedGlyphLine& PrevLine = ShapedText->Lines[Index];
+		const FGlyphLine& PrevLine = GlyphText->Lines[Index];
 		VerticalOffset -= (PrevLine.MaxFontHeight + LineSpacing);
 	}
 
@@ -437,6 +453,12 @@ FVector UText3DDefaultLayoutExtension::GetLineLocation(int32 LineIndex) const
 }
 
 #if WITH_EDITOR
+void UText3DDefaultLayoutExtension::PostEditUndo()
+{
+	Super::PostEditUndo();
+	OnLayoutOptionsChanged();
+}
+
 void UText3DDefaultLayoutExtension::PostEditChangeProperty(FPropertyChangedEvent& InEvent)
 {
 	Super::PostEditChangeProperty(InEvent);
@@ -464,15 +486,17 @@ void UText3DDefaultLayoutExtension::PostEditChangeProperty(FPropertyChangedEvent
 
 void UText3DDefaultLayoutExtension::OnLayoutOptionsChanged()
 {
+	using namespace UE::Text3D::Layout;
+
 	MaxWidth = FMath::Max(1.0f, MaxWidth);
 	MaxHeight = FMath::Max(1.0f, MaxHeight);
 
 	EText3DRendererFlags Flags = EText3DRendererFlags::Material;
 	EnumAddFlags(Flags, EText3DRendererFlags::Layout);
 
-	if (MaxWidthBehavior == EText3DMaxWidthHandling::WrapAndScale && ShapedText.IsValid())
+	if (MaxWidthBehavior == EText3DMaxWidthHandling::WrapAndScale && GlyphText.IsValid())
 	{
-		for (const TConstEnumerateRef<FShapedGlyphLine> ShapedLine : EnumerateRange(ShapedText->Lines))
+		for (const TConstEnumerateRef<FGlyphLine> ShapedLine : EnumerateRange(GlyphText->Lines))
 		{
 			if (ShapedLine->TextDirection == TextBiDi::ETextDirection::RightToLeft)
 			{
diff --git a/Engine/Plugins/Experimental/Text3D/Source/Text3D/Private/Extensions/Text3DDefaultMaterialExtension.cpp b/Engine/Plugins/Experimental/Text3D/Source/Text3D/Private/Extensions/Text3DDefaultMaterialExtension.cpp
--- a/Engine/Plugins/Experimental/Text3D/Source/Text3D/Private/Extensions/Text3DDefaultMaterialExtension.cpp
+++ b/Engine/Plugins/Experimental/Text3D/Source/Text3D/Private/Extensions/Text3DDefaultMaterialExtension.cpp
@@ -4,11 +4,14 @@
 
 #include "Engine/Texture2D.h"
 #include "Extensions/Text3DLayoutExtensionBase.h"
+#include "Extensions/Text3DStyleExtensionBase.h"
 #include "Logs/Text3DLogs.h"
 #include "Materials/Material.h"
 #include "Materials/MaterialInstanceDynamic.h"
 #include "Materials/MaterialInterface.h"
+#include "Misc/EnumerateRange.h"
 #include "Settings/Text3DProjectSettings.h"
+#include "Styles/Text3DStyleBase.h"
 #include "Text3DComponent.h"
 #include "UObject/Package.h"
 
@@ -283,87 +286,136 @@ EText3DExtensionResult UText3DDefaultMaterialExtension::PreRendererUpdate(const
 		return EText3DExtensionResult::Active;
 	}
 
-	if (Style == EText3DMaterialStyle::Custom || Style == EText3DMaterialStyle::Invalid)
+	const UText3DComponent* Text3DComponent = GetText3DComponent();
+
+	if (Style != EText3DMaterialStyle::Custom && Style != EText3DMaterialStyle::Invalid)
 	{
-		return EText3DExtensionResult::Finished;
-	}
+		UMaterialInstanceDynamic* DynFrontMaterial = FindOrAdd(EText3DGroupType::Front);
+		UMaterialInstanceDynamic* DynBackMaterial = FindOrAdd(EText3DGroupType::Back);
+		UMaterialInstanceDynamic* DynExtrudeMaterial = FindOrAdd(EText3DGroupType::Extrude);
+		UMaterialInstanceDynamic* DynBevelMaterial = FindOrAdd(EText3DGroupType::Bevel);
 
-	UMaterialInstanceDynamic* DynFrontMaterial = FindOrAdd(EText3DGroupType::Front);
-	UMaterialInstanceDynamic* DynBackMaterial = FindOrAdd(EText3DGroupType::Back);
-	UMaterialInstanceDynamic* DynExtrudeMaterial = FindOrAdd(EText3DGroupType::Extrude);
-	UMaterialInstanceDynamic* DynBevelMaterial = FindOrAdd(EText3DGroupType::Bevel);
+		if (!DynFrontMaterial || !DynBackMaterial || !DynExtrudeMaterial || !DynBevelMaterial)
+		{
+			UE_LOG(LogText3D, Error, TEXT("Failed to retrieve dynamic material in Text3D material extension"))
+			return EText3DExtensionResult::Failed;
+		}
 
-	if (!DynFrontMaterial || !DynBackMaterial || !DynExtrudeMaterial || !DynBevelMaterial)
-	{
-		UE_LOG(LogText3D, Error, TEXT("Failed to retrieve dynamic material in Text3D material extension"))
-		return EText3DExtensionResult::Failed;
-	}
+		TArray<UMaterialInstanceDynamic*> Materials
+		{
+			DynFrontMaterial,
+			DynBackMaterial,
+			DynExtrudeMaterial,
+			DynBevelMaterial
+		};
 
-	TArray<UMaterialInstanceDynamic*> Materials
-	{
-		DynFrontMaterial,
-		DynBackMaterial,
-		DynExtrudeMaterial,
-		DynBevelMaterial
-	};
+		switch (Style)
+		{
+		case EText3DMaterialStyle::Solid:
+			DynFrontMaterial->SetVectorParameterValue(UText3DProjectSettings::FMaterialParameters::SolidColor, FrontColor);
+			DynBackMaterial->SetVectorParameterValue(UText3DProjectSettings::FMaterialParameters::SolidColor, BackColor);
+			DynExtrudeMaterial->SetVectorParameterValue(UText3DProjectSettings::FMaterialParameters::SolidColor, ExtrudeColor);
+			DynBevelMaterial->SetVectorParameterValue(UText3DProjectSettings::FMaterialParameters::SolidColor, BevelColor);
+			SetScalarParameter(Materials, UText3DProjectSettings::FMaterialParameters::Mode, 0);
+			break;
+		case EText3DMaterialStyle::Gradient:
+			SetVectorParameter(Materials, UText3DProjectSettings::FMaterialParameters::GradientColorA, GradientColorA);
+			SetVectorParameter(Materials, UText3DProjectSettings::FMaterialParameters::GradientColorB, GradientColorB);
+			SetScalarParameter(Materials, UText3DProjectSettings::FMaterialParameters::GradientOffset, GradientOffset);
+			SetScalarParameter(Materials, UText3DProjectSettings::FMaterialParameters::GradientSmoothness, GradientSmoothness);
+			SetScalarParameter(Materials, UText3DProjectSettings::FMaterialParameters::GradientRotation, GradientRotation);
+			SetScalarParameter(Materials, UText3DProjectSettings::FMaterialParameters::Mode, 1);
+			break;
+		case EText3DMaterialStyle::Texture:
+			SetScalarParameter(Materials, UText3DProjectSettings::FMaterialParameters::TexturedUTiling, TextureTiling.X);
+			SetScalarParameter(Materials, UText3DProjectSettings::FMaterialParameters::TexturedVTiling, TextureTiling.Y);
+			SetTextureParameter(Materials, UText3DProjectSettings::FMaterialParameters::MainTexture, TextureAsset);
+			SetScalarParameter(Materials, UText3DProjectSettings::FMaterialParameters::Mode, 2);
+			break;
+		default:
+			break;
+		}
 
-	switch (Style)
-	{
-	case EText3DMaterialStyle::Solid:
-		DynFrontMaterial->SetVectorParameterValue(UText3DProjectSettings::FMaterialParameters::SolidColor, FrontColor);
-		DynBackMaterial->SetVectorParameterValue(UText3DProjectSettings::FMaterialParameters::SolidColor, BackColor);
-		DynExtrudeMaterial->SetVectorParameterValue(UText3DProjectSettings::FMaterialParameters::SolidColor, ExtrudeColor);
-		DynBevelMaterial->SetVectorParameterValue(UText3DProjectSettings::FMaterialParameters::SolidColor, BevelColor);
-		SetScalarParameter(Materials, UText3DProjectSettings::FMaterialParameters::Mode, 0);
-		break;
-	case EText3DMaterialStyle::Gradient:
-		SetVectorParameter(Materials, UText3DProjectSettings::FMaterialParameters::GradientColorA, GradientColorA);
-		SetVectorParameter(Materials, UText3DProjectSettings::FMaterialParameters::GradientColorB, GradientColorB);
-		SetScalarParameter(Materials, UText3DProjectSettings::FMaterialParameters::GradientOffset, GradientOffset);
-		SetScalarParameter(Materials, UText3DProjectSettings::FMaterialParameters::GradientSmoothness, GradientSmoothness);
-		SetScalarParameter(Materials, UText3DProjectSettings::FMaterialParameters::GradientRotation, GradientRotation);
-		SetScalarParameter(Materials, UText3DProjectSettings::FMaterialParameters::Mode, 1);
-		break;
-	case EText3DMaterialStyle::Texture:
-		SetScalarParameter(Materials, UText3DProjectSettings::FMaterialParameters::TexturedUTiling, TextureTiling.X);
-		SetScalarParameter(Materials, UText3DProjectSettings::FMaterialParameters::TexturedVTiling, TextureTiling.Y);
-		SetTextureParameter(Materials, UText3DProjectSettings::FMaterialParameters::MainTexture, TextureAsset);
-		SetScalarParameter(Materials, UText3DProjectSettings::FMaterialParameters::Mode, 2);
-		break;
-	default:
-		break;
-	}
+		if (BlendMode == EText3DMaterialBlendMode::Translucent || Style == EText3DMaterialStyle::Gradient)
+		{
+			const FBox LocalBounds = Text3DComponent->GetBounds();
+			const FVector LocalBoundsExtent = Text3DComponent->GetRelativeRotation().UnrotateVector(LocalBounds.GetSize());
 
-	if (BlendMode == EText3DMaterialBlendMode::Translucent || Style == EText3DMaterialStyle::Gradient)
-	{
-		const UText3DComponent* Text3DComponent = GetText3DComponent();
+			const UText3DLayoutExtensionBase* LayoutExtension = Text3DComponent->GetLayoutExtension();
 
-		const FBox LocalBounds = Text3DComponent->GetBounds();
-		const FVector LocalBoundsExtent = Text3DComponent->GetRelativeRotation().UnrotateVector(LocalBounds.GetSize());
+			if (!LayoutExtension)
+			{
+				UE_LOG(LogText3D, Warning, TEXT("Invalid layout extension retrieved in material extension, cannot proceed."))
+				return EText3DExtensionResult::Failed;
+			}
 
-		const UText3DLayoutExtensionBase* LayoutExtension = Text3DComponent->GetLayoutExtension();
-		const FVector LineLocation = LayoutExtension->GetLineLocation(0);
+			const FVector LineLocation = LayoutExtension->GetLineLocation(0);
+			const FVector TextScaleFactor = LayoutExtension->GetTextScale() * Text3DComponent->GetComponentScale();
+			const FVector BoundsSizeScaled = LocalBoundsExtent / TextScaleFactor;
 
-		const FVector TextScaleFactor = LayoutExtension->GetTextScale() * Text3DComponent->GetComponentScale();
-		const FVector BoundsSizeScaled = LocalBoundsExtent / TextScaleFactor;
+			SetVectorParameter(Materials, UText3DProjectSettings::FMaterialParameters::BoundsOrigin, LineLocation);
+			SetVectorParameter(Materials, UText3DProjectSettings::FMaterialParameters::BoundsSize, BoundsSizeScaled);
+		}
 
-		SetVectorParameter(Materials, UText3DProjectSettings::FMaterialParameters::BoundsOrigin, LineLocation);
-		SetVectorParameter(Materials, UText3DProjectSettings::FMaterialParameters::BoundsSize, BoundsSizeScaled);
+		if (BlendMode == EText3DMaterialBlendMode::Translucent)
+		{
+			SetScalarParameter(Materials, UText3DProjectSettings::FMaterialParameters::Opacity, Opacity);
+			SetScalarParameter(Materials, UText3DProjectSettings::FMaterialParameters::MaskEnabled, bUseMask ? 1 : 0);
+			SetScalarParameter(Materials, UText3DProjectSettings::FMaterialParameters::MaskOffset, MaskOffset);
+			SetScalarParameter(Materials, UText3DProjectSettings::FMaterialParameters::MaskRotation, MaskRotation);
+			SetScalarParameter(Materials, UText3DProjectSettings::FMaterialParameters::MaskSmoothness, MaskSmoothness);
+		}
+
+		FrontMaterial = DynFrontMaterial;
+		BackMaterial = DynBackMaterial;
+		ExtrudeMaterial = DynExtrudeMaterial;
+		BevelMaterial = DynBevelMaterial;
 	}
 
-	if (BlendMode == EText3DMaterialBlendMode::Translucent)
+	if (!MaterialOverrides.IsEmpty())
 	{
-		SetScalarParameter(Materials, UText3DProjectSettings::FMaterialParameters::Opacity, Opacity);
-		SetScalarParameter(Materials, UText3DProjectSettings::FMaterialParameters::MaskEnabled, bUseMask ? 1 : 0);
-		SetScalarParameter(Materials, UText3DProjectSettings::FMaterialParameters::MaskOffset, MaskOffset);
-		SetScalarParameter(Materials, UText3DProjectSettings::FMaterialParameters::MaskRotation, MaskRotation);
-		SetScalarParameter(Materials, UText3DProjectSettings::FMaterialParameters::MaskSmoothness, MaskSmoothness);
-	}
+		const UText3DStyleExtensionBase* StyleExtension = Text3DComponent->GetStyleExtension();
 
-	FrontMaterial = DynFrontMaterial;
-	BackMaterial = DynBackMaterial;
-	ExtrudeMaterial = DynExtrudeMaterial;
-	BevelMaterial = DynBevelMaterial;
+		if (!StyleExtension)
+		{
+			UE_LOG(LogText3D, Warning, TEXT("Invalid style extension retrieved in material extension, cannot proceed."))
+			return EText3DExtensionResult::Failed;
+		}
+
+		for (FText3DMaterialOverride& MaterialOverride : MaterialOverrides)
+		{
+			if (const UText3DStyleBase* StyleRule = StyleExtension->GetStyle(MaterialOverride.Tag))
+			{
+				if (StyleRule->GetOverrideFrontColor())
+				{
+					UMaterialInterface* Material = MaterialOverride.Materials[static_cast<int32>(EText3DGroupType::Front)];
+
+					if (!Material)
+					{
+						Material = FindOrAdd(EText3DGroupType::Front, MaterialOverride.Tag);
+					}
+
+					if (UMaterialInstanceDynamic* DynamicMaterial = Cast<UMaterialInstanceDynamic>(Material))
+					{
+						const FLinearColor SolidColor = StyleRule->GetFrontColor();
+						DynamicMaterial->SetVectorParameterValue(UText3DProjectSettings::FMaterialParameters::SolidColor, SolidColor);
+						DynamicMaterial->SetScalarParameterValue(UText3DProjectSettings::FMaterialParameters::Mode, 0);
+						DynamicMaterial->SetScalarParameterValue(UText3DProjectSettings::FMaterialParameters::Opacity, Opacity);
+					}
+
+					MaterialOverride.Materials[static_cast<int32>(EText3DGroupType::Front)] = Material;
+				}
+				else
+				{
+					MaterialOverride.Materials[static_cast<int32>(EText3DGroupType::Front)] = FrontMaterial;
+				}
+
+				MaterialOverride.Materials[static_cast<int32>(EText3DGroupType::Bevel)] = BevelMaterial;
+				MaterialOverride.Materials[static_cast<int32>(EText3DGroupType::Extrude)] = ExtrudeMaterial;
+				MaterialOverride.Materials[static_cast<int32>(EText3DGroupType::Back)] = BackMaterial;
+			}
+		}
+	}
 
 	return EText3DExtensionResult::Finished;
 }
@@ -373,9 +425,30 @@ EText3DExtensionResult UText3DDefaultMaterialExtension::PostRendererUpdate(const
 	return EText3DExtensionResult::Active;
 }
 
-void UText3DDefaultMaterialExtension::SetMaterial(EText3DGroupType InGroup, UMaterialInterface* InMaterial)
+void UText3DDefaultMaterialExtension::SetMaterial(const UE::Text3D::Material::FMaterialParameters& InParameters, UMaterialInterface* InMaterial)
 {
-	switch (InGroup)
+	if (!InParameters.Tag.IsNone())
+	{
+		FText3DMaterialOverride* MaterialOverride = MaterialOverrides.FindByPredicate([&InParameters](const FText3DMaterialOverride& InMaterialOverride)
+		{
+			return InMaterialOverride.Tag.IsEqual(InParameters.Tag);
+		});
+
+		if (MaterialOverride)
+		{
+			const int32 Index = static_cast<int32>(InParameters.Group);
+
+			if (MaterialOverride->Materials[Index] != InMaterial)
+			{
+				MaterialOverride->Materials[Index] = InMaterial;
+				OnMaterialOptionsChanged();
+			}
+
+			return;
+		}
+	}
+
+	switch (InParameters.Group)
 	{
 	case EText3DGroupType::Front:
 		SetFrontMaterial(InMaterial);
@@ -394,28 +467,90 @@ void UText3DDefaultMaterialExtension::SetMaterial(EText3DGroupType InGroup, UMat
 	}
 }
 
-UMaterialInstanceDynamic* UText3DDefaultMaterialExtension::FindOrAdd(EText3DGroupType InGroup)
+UMaterialInterface* UText3DDefaultMaterialExtension::GetMaterial(const UE::Text3D::Material::FMaterialParameters& InParameters) const
+{
+	if (!InParameters.Tag.IsNone())
+	{
+		const FText3DMaterialOverride* MaterialOverride = MaterialOverrides.FindByPredicate([&InParameters](const FText3DMaterialOverride& InMaterialOverride)
+		{
+			return InMaterialOverride.Tag.IsEqual(InParameters.Tag);
+		});
+
+		const int32 MaterialIndex = static_cast<int32>(InParameters.Group);
+		if (MaterialOverride && MaterialOverride->Materials.IsValidIndex(MaterialIndex))
+		{
+			if (const TObjectPtr<UMaterialInterface> Material = MaterialOverride->Materials[MaterialIndex])
+			{
+				return Material;
+			}
+		}
+	}
+
+	switch (InParameters.Group)
+	{
+	case EText3DGroupType::Front:
+		return FrontMaterial;
+		break;
+	case EText3DGroupType::Bevel:
+		return BevelMaterial;
+		break;
+	case EText3DGroupType::Extrude:
+		return ExtrudeMaterial;
+		break;
+	case EText3DGroupType::Back:
+		return BackMaterial;
+		break;
+	default:
+		return nullptr;
+	}
+}
+
+int32 UText3DDefaultMaterialExtension::GetMaterialCount() const
+{
+	return static_cast<int32>(EText3DGroupType::TypeCount) * (MaterialOverrides.Num() + 1); 
+}
+
+void UText3DDefaultMaterialExtension::GetMaterialNames(TArray<FName>& OutNames) const
+{
+	OutNames.Reset(GetMaterialCount());
+
+	// Default
+	OutNames.Add(NAME_None);
+	Algo::Transform(MaterialOverrides, OutNames, [](const FText3DMaterialOverride& InMaterialOverride)
+	{
+		return InMaterialOverride.Tag;
+	});
+}
+
+UMaterialInstanceDynamic* UText3DDefaultMaterialExtension::FindOrAdd(EText3DGroupType InGroup, FName InTag)
 {
+	using namespace UE::Text3D::Material;
+
 	const UText3DProjectSettings* Text3DSettings = UText3DProjectSettings::Get();
 	check(!!Text3DSettings);
 
 	const FText3DMaterialKey MaterialKey(BlendMode, bIsUnlit);
 	UMaterial* ParentMaterial = Text3DSettings->GetBaseMaterial(MaterialKey);
 
-	if (UMaterialInstanceDynamic* DynMat = Cast<UMaterialInstanceDynamic>(GetMaterial(InGroup)))
+	if (InTag.IsNone())
 	{
-		if (DynMat->GetBaseMaterial() == ParentMaterial)
+		FMaterialParameters Parameters;
+		Parameters.Group = InGroup;
+		if (UMaterialInstanceDynamic* DynMat = Cast<UMaterialInstanceDynamic>(GetMaterial(Parameters)))
 		{
-			return DynMat;
+			if (DynMat->GetBaseMaterial() == ParentMaterial)
+			{
+				return DynMat;
+			}
 		}
-	}
 
-	if (Style == EText3DMaterialStyle::Custom || Style == EText3DMaterialStyle::Invalid)
-	{
-		return nullptr;
+		if (Style == EText3DMaterialStyle::Custom || Style == EText3DMaterialStyle::Invalid)
+		{
+			return nullptr;
+		}
 	}
 
-	FText3DMaterialGroupKey GroupKey(MaterialKey, InGroup, Style);
+	FText3DMaterialGroupKey GroupKey(MaterialKey, InGroup, InTag);
 	if (const TObjectPtr<UMaterialInstanceDynamic>* Material = GroupDynamicMaterials.Find(GroupKey))
 	{
 		return *Material;
@@ -428,6 +563,14 @@ UMaterialInstanceDynamic* UText3DDefaultMaterialExtension::FindOrAdd(EText3DGrou
 	return DynamicMaterial;
 }
 
+UMaterialInstanceDynamic* UText3DDefaultMaterialExtension::Find(EText3DGroupType InGroup, FName InTag) const
+{
+	const FText3DMaterialKey MaterialKey(BlendMode, bIsUnlit);
+	const FText3DMaterialGroupKey GroupKey(MaterialKey, InGroup, InTag);
+	const TObjectPtr<UMaterialInstanceDynamic>* Material = GroupDynamicMaterials.Find(GroupKey);
+	return Material ? *Material : nullptr;
+}
+
 void UText3DDefaultMaterialExtension::SetVectorParameter(TConstArrayView<UMaterialInstanceDynamic*> InMaterials, FName InKey, FVector InValue)
 {
 	for (UMaterialInstanceDynamic* Material : InMaterials)
@@ -477,24 +620,54 @@ void UText3DDefaultMaterialExtension::OnCustomMaterialChanged()
 	OnMaterialOptionsChanged();
 }
 
-UMaterialInterface* UText3DDefaultMaterialExtension::GetMaterial(EText3DGroupType InGroup) const
+void UText3DDefaultMaterialExtension::ClearMaterialOverrides()
 {
-	switch (InGroup)
+	MaterialOverrides.Empty();
+}
+
+void UText3DDefaultMaterialExtension::RegisterMaterialOverride(FName InTag)
+{
+	if (!InTag.IsNone())
 	{
-	case EText3DGroupType::Front:
-		return FrontMaterial;
-		break;
-	case EText3DGroupType::Bevel:
-		return BevelMaterial;
-		break;
-	case EText3DGroupType::Extrude:
-		return ExtrudeMaterial;
-		break;
-	case EText3DGroupType::Back:
-		return BackMaterial;
-		break;
-	default:
-		return nullptr;
+		FText3DMaterialOverride MaterialOverride
+		{
+			.Tag = InTag
+		};
+
+		MaterialOverride.Materials.SetNum(static_cast<int32>(EText3DGroupType::TypeCount));
+
+		MaterialOverrides.AddUnique(MaterialOverride);
+	}
+}
+
+void UText3DDefaultMaterialExtension::ForEachMaterial(TFunctionRef<bool(const UE::Text3D::Material::FMaterialParameters&, UMaterialInterface*)> InFunctor) const
+{
+	// Base materials
+	using namespace UE::Text3D::Material;
+
+	FMaterialParameters Params;
+	for (int32 Index = 0; Index < static_cast<int32>(EText3DGroupType::TypeCount); Index++)
+	{
+		const EText3DGroupType GroupType = static_cast<EText3DGroupType>(Index);
+		Params.Group = GroupType;
+		if (!InFunctor(Params, GetMaterial(Params)))
+		{
+			return;
+		}
+	}
+
+	// Override materials
+	for (const FText3DMaterialOverride& MaterialOverride : MaterialOverrides)
+	{
+		Params.Tag = MaterialOverride.Tag;
+		for (TConstEnumerateRef<TObjectPtr<UMaterialInterface>> Material : EnumerateRange(MaterialOverride.Materials))
+		{
+			Params.Group = static_cast<EText3DGroupType>(Material.GetIndex());
+			if (!InFunctor(Params, *Material))
+			{
+				return;
+			}
+		}
 	}
 }
 
@@ -597,4 +770,10 @@ void UText3DDefaultMaterialExtension::PostEditChangeProperty(FPropertyChangedEve
 		OnMaterialOptionsChanged();
 	}
 }
+
+void UText3DDefaultMaterialExtension::PostEditUndo()
+{
+	Super::PostEditUndo();
+	OnMaterialOptionsChanged();
+}
 #endif
diff --git a/Engine/Plugins/Experimental/Text3D/Source/Text3D/Private/Extensions/Text3DDefaultStyleExtension.cpp b/Engine/Plugins/Experimental/Text3D/Source/Text3D/Private/Extensions/Text3DDefaultStyleExtension.cpp
new file mode 100644
--- /dev/null
+++ b/Engine/Plugins/Experimental/Text3D/Source/Text3D/Private/Extensions/Text3DDefaultStyleExtension.cpp
@@ -0,0 +1,181 @@
+// Copyright Epic Games, Inc. All Rights Reserved.
+
+#include "Extensions/Text3DDefaultStyleExtension.h"
+
+#include "Engine/Font.h"
+#include "Extensions/Text3DMaterialExtensionBase.h"
+#include "Fonts/SlateFontInfo.h"
+#include "Settings/Text3DProjectSettings.h"
+#include "Styles/Text3DStyleBase.h"
+#include "Styles/Text3DStyleSet.h"
+#include "Styling/SlateStyle.h"
+#include "Text3DComponent.h"
+
+UText3DDefaultStyleExtension::UText3DDefaultStyleExtension()
+{
+	if (const UText3DProjectSettings* TextSettings = UText3DProjectSettings::Get())
+	{
+		StyleSet = TextSettings->GetDefaultStyleSet();		
+	}
+
+	UText3DStyleSet::OnStyleSetUpdated().AddUObject(this, &UText3DDefaultStyleExtension::OnStyleSetUpdated);
+}
+
+void UText3DDefaultStyleExtension::SetStyleSet(UText3DStyleSet* InStyleSet)
+{
+	if (StyleSet == InStyleSet)
+	{
+		return;
+	}
+
+	StyleSet = InStyleSet;
+	OnStyleSetChanged();
+}
+
+#if WITH_EDITOR
+void UText3DDefaultStyleExtension::PostEditUndo()
+{
+	Super::PostEditUndo();
+	RequestUpdate(EText3DRendererFlags::All, /** Immediate */false);
+}
+
+void UText3DDefaultStyleExtension::PostEditChangeProperty(FPropertyChangedEvent& InPropertyChangedEvent)
+{
+	Super::PostEditChangeProperty(InPropertyChangedEvent);
+
+	const FName MemberPropertyName = InPropertyChangedEvent.GetMemberPropertyName();
+
+	if (MemberPropertyName == GET_MEMBER_NAME_CHECKED(UText3DDefaultStyleExtension, Styles))
+	{
+		if (InPropertyChangedEvent.ChangeType == EPropertyChangeType::ArrayClear
+			|| InPropertyChangedEvent.ChangeType == EPropertyChangeType::ArrayRemove)
+		{
+			RequestUpdate(EText3DRendererFlags::All, /** Immediate */false);
+		}
+	}
+	else if (MemberPropertyName == GET_MEMBER_NAME_CHECKED(UText3DDefaultStyleExtension, StyleSet))
+	{
+		OnStyleSetChanged();
+	}
+}
+#endif
+
+TSharedPtr<FTextBlockStyle> UText3DDefaultStyleExtension::GetDefaultStyle() const
+{
+	return DefaultFontStyle;
+}
+
+TSharedPtr<FSlateStyleSet> UText3DDefaultStyleExtension::GetCustomStyles() const
+{
+	return TextStyleSet;
+}
+
+UText3DStyleBase* UText3DDefaultStyleExtension::GetStyle(FName InStyle) const
+{
+	// Look for overrides first
+	for (UText3DStyleBase* FormatStyle : Styles)
+	{
+		const FName StyleName = FormatStyle->GetStyleName();
+		if (StyleName.IsEqual(InStyle))
+		{
+			return FormatStyle;
+		}
+	}
+
+	if (StyleSet)
+	{
+		for (UText3DStyleBase* FormatStyle : StyleSet->GetStyles())
+		{
+			const FName StyleName = FormatStyle->GetStyleName();
+			if (StyleName.IsEqual(InStyle))
+			{
+				return FormatStyle;
+			}
+		}
+	}
+
+	return nullptr;
+}
+
+EText3DExtensionResult UText3DDefaultStyleExtension::PreRendererUpdate(const UE::Text3D::Renderer::FUpdateParameters& InParameters)
+{
+	if (!DefaultFontStyle.IsValid()
+		|| !TextStyleSet.IsValid()
+		|| EnumHasAnyFlags(InParameters.UpdateFlags, EText3DRendererFlags::Geometry))
+	{
+		const UText3DComponent* Text3DComponent = GetText3DComponent();
+		const UFont* Font = Text3DComponent->GetFont();
+		UText3DMaterialExtensionBase* MaterialExtension = Text3DComponent->GetMaterialExtension();
+
+		if (!DefaultFontStyle.IsValid())
+		{
+			DefaultFontStyle = MakeShared<FTextBlockStyle>();
+		}
+
+		FSlateFontInfo FontInfo(Font, UE::Text3D::Metrics::SlateFontSize);
+		FontInfo.CompositeFont = Font && Font->GetCompositeFont() ? MakeShared<FCompositeFont>(*Font->GetCompositeFont()) : FStyleDefaults::GetFontInfo().CompositeFont;
+		FontInfo.TypefaceFontName = Text3DComponent->GetTypeface();
+		FontInfo.FontFallback = EFontFallback::FF_NoFallback;
+		DefaultFontStyle->SetFont(FontInfo);
+
+		TextStyleSet = MakeShared<FSlateStyleSet>(TEXT("Text3DRichTextStyle"));
+		TextStyleSet->Set(FName(TEXT("default")), *DefaultFontStyle);
+
+		MaterialExtension->ClearMaterialOverrides();
+
+		// Process style set first
+		if (StyleSet)
+		{
+			for (const UText3DStyleBase* FormatStyle : StyleSet->GetStyles())
+			{
+				if (FormatStyle)
+				{
+					const FName StyleName = FormatStyle->GetStyleName();
+					if (!StyleName.IsNone())
+					{
+						FTextBlockStyle TextStyle = *DefaultFontStyle;
+						FormatStyle->ConfigureStyle(TextStyle);
+						TextStyleSet->Set(StyleName, TextStyle);
+						MaterialExtension->RegisterMaterialOverride(StyleName);
+					}
+				}
+			}
+		}
+
+		// Process additional styles to override style set
+		for (const UText3DStyleBase* FormatStyle : Styles)
+		{
+			if (FormatStyle)
+			{
+				const FName StyleName = FormatStyle->GetStyleName();
+				if (!StyleName.IsNone())
+				{
+					FTextBlockStyle TextStyle = *DefaultFontStyle;
+					FormatStyle->ConfigureStyle(TextStyle);
+					TextStyleSet->Set(StyleName, TextStyle);
+					MaterialExtension->RegisterMaterialOverride(StyleName);
+				}
+			}
+		}
+	}
+
+	return EText3DExtensionResult::Finished;
+}
+
+EText3DExtensionResult UText3DDefaultStyleExtension::PostRendererUpdate(const UE::Text3D::Renderer::FUpdateParameters& InParameters)
+{
+	return EText3DExtensionResult::Active;
+}
+
+void UText3DDefaultStyleExtension::OnStyleSetUpdated(UText3DStyleSet* InStyleSet)
+{
+	if (StyleSet == InStyleSet)
+	{
+		OnStyleSetChanged();
+	}
+}
+
+void UText3DDefaultStyleExtension::OnStyleSetChanged()
+{
+	RequestUpdate(EText3DRendererFlags::All, /** Immediate */false);
+}
diff --git a/Engine/Plugins/Experimental/Text3D/Source/Text3D/Private/Extensions/Text3DDefaultTemplateExtension.cpp b/Engine/Plugins/Experimental/Text3D/Source/Text3D/Private/Extensions/Text3DDefaultTemplateExtension.cpp
deleted file mode 100644
diff --git a/Engine/Plugins/Experimental/Text3D/Source/Text3D/Private/Extensions/Text3DDefaultTokenExtension.cpp b/Engine/Plugins/Experimental/Text3D/Source/Text3D/Private/Extensions/Text3DDefaultTokenExtension.cpp
new file mode 100644
--- /dev/null
+++ b/Engine/Plugins/Experimental/Text3D/Source/Text3D/Private/Extensions/Text3DDefaultTokenExtension.cpp
@@ -0,0 +1,67 @@
+// Copyright Epic Games, Inc. All Rights Reserved.
+
+#include "Extensions/Text3DDefaultTokenExtension.h"
+
+#include "Internationalization/TextFormatter.h"
+#include "Text3DComponent.h"
+#include "Tokens/Text3DTokenBase.h"
+
+#if WITH_EDITOR
+void UText3DDefaultTokenExtension::PostEditUndo()
+{
+	Super::PostEditUndo();
+	RequestUpdate(EText3DRendererFlags::All, /** Immediate */false);
+}
+
+void UText3DDefaultTokenExtension::PostEditChangeProperty(FPropertyChangedEvent& InPropertyChangedEvent)
+{
+	Super::PostEditChangeProperty(InPropertyChangedEvent);
+
+	const FName MemberPropertyName = InPropertyChangedEvent.GetMemberPropertyName();
+
+	if (MemberPropertyName == GET_MEMBER_NAME_CHECKED(UText3DDefaultTokenExtension, Tokens))
+	{
+		if (InPropertyChangedEvent.ChangeType == EPropertyChangeType::ArrayClear
+			|| InPropertyChangedEvent.ChangeType == EPropertyChangeType::ArrayRemove
+			|| InPropertyChangedEvent.ChangeType == EPropertyChangeType::ArrayMove)
+		{
+			RequestUpdate(EText3DRendererFlags::All, /** Immediate */false);
+		}
+	}
+}
+#endif
+
+const FText& UText3DDefaultTokenExtension::GetFormattedText() const
+{
+	return FormattedText;
+}
+
+EText3DExtensionResult UText3DDefaultTokenExtension::PreRendererUpdate(const UE::Text3D::Renderer::FUpdateParameters& InParameters)
+{
+	const UText3DComponent* Text3DComponent = GetText3DComponent();
+	FormattedText = Text3DComponent->GetText();
+
+	FFormatNamedArguments NamedArguments;
+	NamedArguments.Reserve(Tokens.Num() * 2);
+	for (UText3DTokenBase* Token : Tokens)
+	{
+		if (Token)
+		{
+			Token->CollectTokens(NamedArguments);
+		}
+	}
+
+	constexpr bool bRebuildText = false;
+	constexpr bool bRebuildAsSource = false;
+	FormattedText = FTextFormatter::Format(MoveTemp(FormattedText), MoveTemp(NamedArguments), bRebuildText, bRebuildAsSource);
+
+	// For backward compatibility
+	Text3DComponent->FormatText(FormattedText);
+
+	return EText3DExtensionResult::Finished;
+}
+
+EText3DExtensionResult UText3DDefaultTokenExtension::PostRendererUpdate(const UE::Text3D::Renderer::FUpdateParameters& InParameters)
+{
+	return EText3DExtensionResult::Active;
+}
diff --git a/Engine/Plugins/Experimental/Text3D/Source/Text3D/Private/GeometryBuilders/Text3DGlyph.cpp b/Engine/Plugins/Experimental/Text3D/Source/Text3D/Private/GeometryBuilders/Text3DGlyph.cpp
--- a/Engine/Plugins/Experimental/Text3D/Source/Text3D/Private/GeometryBuilders/Text3DGlyph.cpp
+++ b/Engine/Plugins/Experimental/Text3D/Source/Text3D/Private/GeometryBuilders/Text3DGlyph.cpp
@@ -33,7 +33,7 @@ void FText3DGlyph::Build(FText3DCachedMesh& InMesh, UMaterialInterface* InDefaul
 
 	auto AddMaterial = [](UStaticMesh* InMesh, UMaterialInterface* InMaterial, int32 InIndex)
 	{
-		using namespace UE::Text3D::Materials;
+		using namespace UE::Text3D::Material;
 
 		check(SlotNames.IsValidIndex(InIndex))
 
diff --git a/Engine/Plugins/Experimental/Text3D/Source/Text3D/Private/LayoutBuilders/Text3DLayoutShaper.cpp b/Engine/Plugins/Experimental/Text3D/Source/Text3D/Private/LayoutBuilders/Text3DLayoutShaper.cpp
--- a/Engine/Plugins/Experimental/Text3D/Source/Text3D/Private/LayoutBuilders/Text3DLayoutShaper.cpp
+++ b/Engine/Plugins/Experimental/Text3D/Source/Text3D/Private/LayoutBuilders/Text3DLayoutShaper.cpp
@@ -26,8 +26,10 @@ void FText3DLayoutShaper::ShapeBidirectionalText(
 	const FString& Text,
 	const TSharedPtr<FTextLayout>& TextLayout,
 	const TSharedPtr<FRichTextLayoutMarshaller>& TextMarshaller,
-	TArray<FShapedGlyphLine>& OutShapedLines)
+	TArray<UE::Text3D::Layout::FGlyphLine>& OutShapedLines)
 {
+	using namespace UE::Text3D::Layout;
+
 	// @todo: Restore when dependency on SlateApplication is removed. Currently this means text meshes won't be created on the server.
 	// @see: UE-211843
 	if (!FSlateApplication::IsInitialized())
@@ -52,7 +54,7 @@ void FText3DLayoutShaper::ShapeBidirectionalText(
 	// @note: mimics FSlateTextLayout::OnPaint
 	for (const FTextLayout::FLineView& Line : LineViews)
 	{
-		FShapedGlyphLine& ShapedLine = OutShapedLines.AddDefaulted_GetRef();
+		FGlyphLine& ShapedLine = OutShapedLines.AddDefaulted_GetRef();
 		ShapedLine.TextDirection = Line.TextBaseDirection;
 
 		for (const TSharedRef<ILayoutBlock>& Block : Line.Blocks)
@@ -73,12 +75,16 @@ void FText3DLayoutShaper::ShapeBidirectionalText(
 				BlockTextContext.TextDirection
 			);
 
-			ShapedLine.GlyphsToRender.Append(ShapedGlyphSequence->GetGlyphsToRender());
+			for (const FShapedGlyphEntry& ShapedGlyphEntry : ShapedGlyphSequence->GetGlyphsToRender())
+			{
+				FGlyphEntry& GlyphEntry = ShapedLine.Glyphs.Add_GetRef(ShapedGlyphEntry);
+				GlyphEntry.StyleTag = StyleTag;
+			}
 		}
 
 		if (ShapedLine.TextDirection == TextBiDi::ETextDirection::RightToLeft)
 		{
-			Algo::Reverse(ShapedLine.GlyphsToRender);
+			Algo::Reverse(ShapedLine.Glyphs);
 		}
 	}
 }
diff --git a/Engine/Plugins/Experimental/Text3D/Source/Text3D/Private/LayoutBuilders/Text3DLayoutShaper.h b/Engine/Plugins/Experimental/Text3D/Source/Text3D/Private/LayoutBuilders/Text3DLayoutShaper.h
--- a/Engine/Plugins/Experimental/Text3D/Source/Text3D/Private/LayoutBuilders/Text3DLayoutShaper.h
+++ b/Engine/Plugins/Experimental/Text3D/Source/Text3D/Private/LayoutBuilders/Text3DLayoutShaper.h
@@ -12,6 +12,11 @@ class FTextLayout;
 struct FShapedGlyphLine;
 struct FTextBlockStyle;
 
+namespace UE::Text3D::Layout
+{
+	struct FGlyphLine;
+}
+
 /** A singleton that handles shaping operations and writes the result to a provided TextLayout. */
 class FText3DLayoutShaper final
 {
@@ -29,7 +34,7 @@ public:
 		const FString& Text,
 		const TSharedPtr<FTextLayout>& TextLayout,
 		const TSharedPtr<FRichTextLayoutMarshaller>& TextMarshaller,
-		TArray<FShapedGlyphLine>& OutShapedLines);
+		TArray<UE::Text3D::Layout::FGlyphLine>& OutShapedLines);
 
 private:
 	struct FPrivateToken { explicit FPrivateToken() = default; };
diff --git a/Engine/Plugins/Experimental/Text3D/Source/Text3D/Private/LayoutBuilders/Text3DShapedGlyphLine.cpp b/Engine/Plugins/Experimental/Text3D/Source/Text3D/Private/LayoutBuilders/Text3DShapedGlyphLine.cpp
--- a/Engine/Plugins/Experimental/Text3D/Source/Text3D/Private/LayoutBuilders/Text3DShapedGlyphLine.cpp
+++ b/Engine/Plugins/Experimental/Text3D/Source/Text3D/Private/LayoutBuilders/Text3DShapedGlyphLine.cpp
@@ -2,15 +2,15 @@
 
 #include "LayoutBuilders/Text3DShapedGlyphLine.h"
 
-float FShapedGlyphLine::GetWidthAdvance(const int32 InIndex, const float InKerning, const float InWordSpacing) const
+float UE::Text3D::Layout::FGlyphLine::GetWidthAdvance(const int32 InIndex, const float InKerning, const float InWordSpacing) const
 {
-	check(InIndex >= 0 && InIndex < GlyphsToRender.Num());
+	check(InIndex >= 0 && InIndex < Glyphs.Num());
 
-	const FShapedGlyphEntry& Glyph = GlyphsToRender[InIndex];
-	float Advance = Glyph.XOffset + Glyph.XAdvance;
+	const FGlyphEntry& Glyph = Glyphs[InIndex];
+	float Advance = Glyph.Entry.XOffset + Glyph.Entry.XAdvance;
 
 	const bool bFirstCharacter = InIndex == 0;
-	const bool bLastCharacter = InIndex == GlyphsToRender.Num() - 1;
+	const bool bLastCharacter = InIndex == Glyphs.Num() - 1;
 	const bool bSkipCharacter = TextDirection == TextBiDi::ETextDirection::RightToLeft ? bLastCharacter : bFirstCharacter;
 
 	if (!bSkipCharacter)
@@ -18,7 +18,7 @@ float FShapedGlyphLine::GetWidthAdvance(const int32 InIndex, const float InKerni
 		// @note: as per FSlateElementBatcher::BuildShapedTextSequence, per Glyph Kerning isn't used
 		Advance += InKerning;
 
-		if (!Glyph.bIsVisible)
+		if (!Glyph.Entry.bIsVisible)
 		{
 			Advance += InWordSpacing;
 		}
diff --git a/Engine/Plugins/Experimental/Text3D/Source/Text3D/Private/LayoutBuilders/Text3DShapedGlyphLine.h b/Engine/Plugins/Experimental/Text3D/Source/Text3D/Private/LayoutBuilders/Text3DShapedGlyphLine.h
--- a/Engine/Plugins/Experimental/Text3D/Source/Text3D/Private/LayoutBuilders/Text3DShapedGlyphLine.h
+++ b/Engine/Plugins/Experimental/Text3D/Source/Text3D/Private/LayoutBuilders/Text3DShapedGlyphLine.h
@@ -6,30 +6,49 @@
 #include "Fonts/FontCache.h"
 #include "Text3DTypes.h"
 
-/** Contains text line metric with sufficient information to fetch and transform each character. */
-struct FShapedGlyphLine
+namespace UE::Text3D::Layout
 {
-	/** Get's the offset from the previous character, accounting for kerning and word spacing. */
-	float GetWidthAdvance(const int32 InIndex, const float InKerning, const float InWordSpacing) const;
-
-	/** The corresponding shaped glyph for each character in this line of text. */
-	TArray<FShapedGlyphEntry> GlyphsToRender;
-
-	/** Stored result of line width */
-	float Width = 0.0f;
-
-	/** Glyph advance on this line */
-	TArray<float> GlyphAdvances;
-
-	/** Direction of the text in the line */
-	TextBiDi::ETextDirection TextDirection = TextBiDi::ETextDirection::LeftToRight;
-
-	/** Max font height for this line calculated based on MaxFontAscender and MaxFontDescender */
-	float MaxFontHeight = 0.f;
-
-	/** Positive value representing max height above baseline */
-	float MaxFontAscender = 0.f;
-
- 	/** Positive value representing max height below baseline */
-	float MaxFontDescender = 0.f;
-};
+	/** Contains text glyph information for a single character */
+	struct FGlyphEntry
+	{
+		FGlyphEntry() = default;
+		FGlyphEntry(const FShapedGlyphEntry& InEntry)
+			: Entry(InEntry)
+		{
+		}
+
+		/** Style applied on this entry */
+		FName StyleTag = NAME_None;
+
+		/** Actual entry for the glyph */
+		FShapedGlyphEntry Entry;
+	};
+
+	/** Contains text line metric with sufficient information to fetch and transform each character. */
+	struct FGlyphLine
+	{
+		/** Get's the offset from the previous character, accounting for kerning and word spacing. */
+		float GetWidthAdvance(const int32 InIndex, const float InKerning, const float InWordSpacing) const;
+
+		/** The corresponding shaped glyph for each character in this line of text. */
+		TArray<FGlyphEntry> Glyphs;
+
+		/** Stored result of line width */
+		float Width = 0.0f;
+
+		/** Glyph advance on this line */
+		TArray<float> GlyphAdvances;
+
+		/** Direction of the text in the line */
+		TextBiDi::ETextDirection TextDirection = TextBiDi::ETextDirection::LeftToRight;
+
+		/** Max font height for this line calculated based on MaxFontAscender and MaxFontDescender */
+		float MaxFontHeight = 0.f;
+
+		/** Positive value representing max height above baseline */
+		float MaxFontAscender = 0.f;
+
+		/** Positive value representing max height below baseline */
+		float MaxFontDescender = 0.f;
+	};
+}
\ No newline at end of file
diff --git a/Engine/Plugins/Experimental/Text3D/Source/Text3D/Private/LayoutBuilders/Text3DShapedGlyphText.cpp b/Engine/Plugins/Experimental/Text3D/Source/Text3D/Private/LayoutBuilders/Text3DShapedGlyphText.cpp
--- a/Engine/Plugins/Experimental/Text3D/Source/Text3D/Private/LayoutBuilders/Text3DShapedGlyphText.cpp
+++ b/Engine/Plugins/Experimental/Text3D/Source/Text3D/Private/LayoutBuilders/Text3DShapedGlyphText.cpp
@@ -7,137 +7,140 @@
 #include "Fonts/SlateTextShaper.h"
 #endif
 
-bool FText3DShapedGlyphText::IsValidGlyph(const FShapedGlyphEntry& InGlyph)
+namespace UE::Text3D::Layout
 {
-#if WITH_FREETYPE
-	return InGlyph.HasValidGlyph() && InGlyph.bIsVisible && InGlyph.GlyphIndex != 0 && InGlyph.FontFaceData.IsValid() && InGlyph.FontFaceData->FontFace.IsValid();
-#else
-	return false;
-#endif
-}
-
-FText3DShapedGlyphText::FText3DShapedGlyphText()
-{
-	Reset();
-}
-
-void FText3DShapedGlyphText::Reset()
-{
-	Kernings.Reset();
-	FontFaces.Reset();
-	Tracking = 0.f;
-	WordSpacing = 0.0f;
-	bWrap = false;
-	Lines.Reset();
-}
-
-void FText3DShapedGlyphText::CalculateWidth()
-{
-	TArray<FShapedGlyphLine> NewLines;
-    NewLines.Reserve(Lines.Num());
-
-    int32 CharacterIndex = 0;
-    for (const TConstEnumerateRef<FShapedGlyphLine> GlyphLine : EnumerateRange(Lines))
-    {
-        FShapedGlyphLine* CurrentLine = &NewLines.AddDefaulted_GetRef();
-        CurrentLine->TextDirection = GlyphLine->TextDirection;
-
-        const bool bIsRTL = (GlyphLine->TextDirection == TextBiDi::ETextDirection::RightToLeft);
-
-        TArray<float> CurrentAdvances;
-        TArray<FShapedGlyphEntry> CurrentWord;
-        float LineWidth = 0.0f;
-        float CurrentWordLength = 0.0f;
-    	float MaxFontAscender = 0.0f;
-    	float MaxFontDescender = 0.0f;
-
-        const int32 GlyphCount = GlyphLine->GlyphsToRender.Num();
-
-		for (TConstEnumerateRef<FShapedGlyphEntry> CurrentGlyph : EnumerateRange(GlyphLine->GlyphsToRender))
-        {
-    		const bool bVisibleGlyph = FText3DShapedGlyphText::IsValidGlyph(*CurrentGlyph);
-			const int32 Index = bIsRTL ? Kernings.Num() - 1 - CharacterIndex : CharacterIndex;
-    		const float CurrentKerning = bVisibleGlyph ? Kernings[Index] : 0.f;
-    		const float GlyphAdv = GlyphLine->GetWidthAdvance(CurrentGlyph.GetIndex(), Tracking + CurrentKerning, WordSpacing);
-
-            // If we're at the end the line or at whitespace
-            const bool bWordBreak = !bVisibleGlyph || CurrentGlyph.GetIndex() == GlyphCount - 1;
-
-    		if (bWrap                               // when we're wrapping
-				&& bWordBreak                       // and at a word break
-				&& LineWidth > MaxWidth             // and the current line is longer than the max
-				&& CurrentWordLength != LineWidth)  // and the line is not just a single word that we can't break
-    		{
-    			CurrentLine->Width = LineWidth - CurrentWordLength;
-    			CurrentLine->MaxFontHeight = MaxFontAscender + MaxFontDescender;
-    			CurrentLine->MaxFontAscender = MaxFontAscender;
-    			CurrentLine->MaxFontDescender = MaxFontDescender;
-
-    			// Reset
-    			MaxFontAscender = 0.0f;
-    			MaxFontDescender = 0.0f;
-
-    			CurrentLine = &NewLines.AddDefaulted_GetRef();
-    			CurrentLine->TextDirection = GlyphLine->TextDirection;
-    			LineWidth = CurrentWordLength;
-    		}
-
-			LineWidth += GlyphAdv;
-			CurrentWordLength += GlyphAdv;
-
-			if (bVisibleGlyph)
-			{
-#if WITH_FREETYPE
-				const FT_Face Face = FontFaces[Index]->GetFace();
-
-				const float Ascender = FMath::Max(0.0f, Face->size->metrics.ascender * UE::Text3D::Metrics::FontSizeInverse);
-				const float Descender = FMath::Max(0.0f, -Face->size->metrics.descender * UE::Text3D::Metrics::FontSizeInverse);
-
-				MaxFontAscender  = FMath::Max(MaxFontAscender,  Ascender);
-				MaxFontDescender = FMath::Max(MaxFontDescender, Descender);
-#endif
-			}
-
-			if (bIsRTL)
-			{
-				CurrentWord.Insert(*CurrentGlyph, 0);
-				CurrentAdvances.Insert(GlyphAdv, 0);
-			}
-			else
-			{
-				CurrentWord.Add(*CurrentGlyph);
-				CurrentAdvances.Add(GlyphAdv);
-			}
-
-            if (bWordBreak)
-            {
-        		if (bIsRTL)
-        		{
-        			CurrentLine->GlyphsToRender.Insert(CurrentWord, 0);
-        			CurrentLine->GlyphAdvances.Insert(CurrentAdvances, 0);
-        		}
-	            else
+	bool FGlyphText::IsValidGlyph(const FGlyphEntry& InGlyph)
+	{
+	#if WITH_FREETYPE
+		return InGlyph.Entry.HasValidGlyph() && InGlyph.Entry.bIsVisible && InGlyph.Entry.GlyphIndex != 0 && InGlyph.Entry.FontFaceData.IsValid() && InGlyph.Entry.FontFaceData->FontFace.IsValid();
+	#else
+		return false;
+	#endif
+	}
+
+	FGlyphText::FGlyphText()
+	{
+		Reset();
+	}
+
+	void FGlyphText::Reset()
+	{
+		Kernings.Reset();
+		FontFaces.Reset();
+		Tracking = 0.f;
+		WordSpacing = 0.0f;
+		bWrap = false;
+		Lines.Reset();
+	}
+
+	void FGlyphText::CalculateWidth()
+	{
+		TArray<FGlyphLine> NewLines;
+	    NewLines.Reserve(Lines.Num());
+
+	    int32 CharacterIndex = 0;
+	    for (const TConstEnumerateRef<FGlyphLine> GlyphLine : EnumerateRange(Lines))
+	    {
+	        FGlyphLine* CurrentLine = &NewLines.AddDefaulted_GetRef();
+	        CurrentLine->TextDirection = GlyphLine->TextDirection;
+
+	        const bool bIsRTL = (GlyphLine->TextDirection == TextBiDi::ETextDirection::RightToLeft);
+
+	        TArray<float> CurrentAdvances;
+	        TArray<FGlyphEntry> CurrentWord;
+	        float LineWidth = 0.0f;
+	        float CurrentWordLength = 0.0f;
+    		float MaxFontAscender = 0.0f;
+    		float MaxFontDescender = 0.0f;
+
+	        const int32 GlyphCount = GlyphLine->Glyphs.Num();
+
+			for (TConstEnumerateRef<FGlyphEntry> CurrentGlyph : EnumerateRange(GlyphLine->Glyphs))
+	        {
+    			const bool bVisibleGlyph = IsValidGlyph(*CurrentGlyph);
+				const int32 Index = bIsRTL ? Kernings.Num() - 1 - CharacterIndex : CharacterIndex;
+    			const float CurrentKerning = bVisibleGlyph ? Kernings[Index] : 0.f;
+    			const float GlyphAdv = GlyphLine->GetWidthAdvance(CurrentGlyph.GetIndex(), Tracking + CurrentKerning, WordSpacing);
+
+	            // If we're at the end the line or at whitespace
+	            const bool bWordBreak = !bVisibleGlyph || CurrentGlyph.GetIndex() == GlyphCount - 1;
+
+    			if (bWrap                               // when we're wrapping
+					&& bWordBreak                       // and at a word break
+					&& LineWidth > MaxWidth             // and the current line is longer than the max
+					&& CurrentWordLength != LineWidth)  // and the line is not just a single word that we can't break
+    			{
+    				CurrentLine->Width = LineWidth - CurrentWordLength;
+    				CurrentLine->MaxFontHeight = MaxFontAscender + MaxFontDescender;
+    				CurrentLine->MaxFontAscender = MaxFontAscender;
+    				CurrentLine->MaxFontDescender = MaxFontDescender;
+
+    				// Reset
+    				MaxFontAscender = 0.0f;
+    				MaxFontDescender = 0.0f;
+
+    				CurrentLine = &NewLines.AddDefaulted_GetRef();
+    				CurrentLine->TextDirection = GlyphLine->TextDirection;
+    				LineWidth = CurrentWordLength;
+    			}
+
+				LineWidth += GlyphAdv;
+				CurrentWordLength += GlyphAdv;
+
+				if (bVisibleGlyph)
+				{
+	#if WITH_FREETYPE
+					const FT_Face Face = FontFaces[Index]->GetFace();
+
+					const float Ascender = FMath::Max(0.0f, Face->size->metrics.ascender * UE::Text3D::Metrics::FontSizeInverse);
+					const float Descender = FMath::Max(0.0f, -Face->size->metrics.descender * UE::Text3D::Metrics::FontSizeInverse);
+
+					MaxFontAscender  = FMath::Max(MaxFontAscender,  Ascender);
+					MaxFontDescender = FMath::Max(MaxFontDescender, Descender);
+	#endif
+				}
+
+				if (bIsRTL)
+				{
+					CurrentWord.Insert(*CurrentGlyph, 0);
+					CurrentAdvances.Insert(GlyphAdv, 0);
+				}
+				else
+				{
+					CurrentWord.Add(*CurrentGlyph);
+					CurrentAdvances.Add(GlyphAdv);
+				}
+
+	            if (bWordBreak)
 	            {
-            		CurrentLine->GlyphsToRender.Append(CurrentWord);
-            		CurrentLine->GlyphAdvances.Append(CurrentAdvances);
+        			if (bIsRTL)
+        			{
+        				CurrentLine->Glyphs.Insert(CurrentWord, 0);
+        				CurrentLine->GlyphAdvances.Insert(CurrentAdvances, 0);
+        			}
+		            else
+		            {
+            			CurrentLine->Glyphs.Append(CurrentWord);
+            			CurrentLine->GlyphAdvances.Append(CurrentAdvances);
+		            }
+
+	                CurrentWordLength = 0.0f;
+	                CurrentWord.Reset();
+	                CurrentAdvances.Reset();
 	            }
 
-                CurrentWordLength = 0.0f;
-                CurrentWord.Reset();
-                CurrentAdvances.Reset();
-            }
-
-    		if (bVisibleGlyph)
-    		{
-    			CharacterIndex++;
-    		}
-        }
+    			if (bVisibleGlyph)
+    			{
+    				CharacterIndex++;
+    			}
+	        }
 
-        CurrentLine->Width = LineWidth;
-    	CurrentLine->MaxFontHeight = MaxFontAscender + MaxFontDescender;
-    	CurrentLine->MaxFontAscender = MaxFontAscender;
-    	CurrentLine->MaxFontDescender = MaxFontDescender;
-    }
+	        CurrentLine->Width = LineWidth;
+    		CurrentLine->MaxFontHeight = MaxFontAscender + MaxFontDescender;
+    		CurrentLine->MaxFontAscender = MaxFontAscender;
+    		CurrentLine->MaxFontDescender = MaxFontDescender;
+	    }
 
-    Lines = NewLines;
+	    Lines = NewLines;
+	}
 }
\ No newline at end of file
diff --git a/Engine/Plugins/Experimental/Text3D/Source/Text3D/Private/LayoutBuilders/Text3DShapedGlyphText.h b/Engine/Plugins/Experimental/Text3D/Source/Text3D/Private/LayoutBuilders/Text3DShapedGlyphText.h
--- a/Engine/Plugins/Experimental/Text3D/Source/Text3D/Private/LayoutBuilders/Text3DShapedGlyphText.h
+++ b/Engine/Plugins/Experimental/Text3D/Source/Text3D/Private/LayoutBuilders/Text3DShapedGlyphText.h
@@ -8,36 +8,39 @@
 
 class FFreeTypeFace;
 struct FShapedGlyphLine;
-struct FT_FaceRec_;
 
-/** Contains text metrics with sufficient information to fetch and transform each line. */
-struct FText3DShapedGlyphText
+namespace UE::Text3D::Layout
 {
-	static bool IsValidGlyph(const FShapedGlyphEntry& InGlyph);
+	/** Contains text metrics with sufficient information to fetch and transform each line. */
+	struct FGlyphText
+	{
+		/** Check if a glyph is valid and visible (can be rendered) */
+		static bool IsValidGlyph(const FGlyphEntry& InGlyph);
 
-	FText3DShapedGlyphText();
+		FGlyphText();
 
-	void Reset();
-	void CalculateWidth();
+		void Reset();
+		void CalculateWidth();
 
-	/** Individual character kerning */
-	TArray<float> Kernings;
+		/** Individual character kerning */
+		TArray<float> Kernings;
 
-	/** Individual character font faces */
-	TArray<TSharedPtr<FFreeTypeFace>> FontFaces;
+		/** Individual character font faces */
+		TArray<TSharedPtr<FFreeTypeFace>> FontFaces;
 
-	/** General tracking for all characters */
-	float Tracking;
+		/** General tracking for all characters */
+		float Tracking;
 
-	/** Spacing between words */
-	float WordSpacing;
+		/** Spacing between words */
+		float WordSpacing;
 
-	/** Maximum width allowed */
-	float MaxWidth;
+		/** Maximum width allowed */
+		float MaxWidth;
 
-	/** Whether to wrap words */
-	bool bWrap;
+		/** Whether to wrap words */
+		bool bWrap;
 
-	/** Lines calculated based on current text */
-	TArray<FShapedGlyphLine> Lines;
-};
+		/** Lines calculated based on current text */
+		TArray<FGlyphLine> Lines;
+	};
+}
\ No newline at end of file
diff --git a/Engine/Plugins/Experimental/Text3D/Source/Text3D/Private/Renderers/DynamicMesh/Text3DDynamicMeshRenderer.cpp b/Engine/Plugins/Experimental/Text3D/Source/Text3D/Private/Renderers/DynamicMesh/Text3DDynamicMeshRenderer.cpp
--- a/Engine/Plugins/Experimental/Text3D/Source/Text3D/Private/Renderers/DynamicMesh/Text3DDynamicMeshRenderer.cpp
+++ b/Engine/Plugins/Experimental/Text3D/Source/Text3D/Private/Renderers/DynamicMesh/Text3DDynamicMeshRenderer.cpp
@@ -116,7 +116,7 @@ void UText3DDynamicMeshRenderer::OnUpdate(const UE::Text3D::Renderer::FUpdatePar
 	}
 	else if (InParameters.CurrentFlag == EText3DRendererFlags::Material)
 	{
-		using namespace UE::Text3D::Materials;
+		using namespace UE::Text3D::Material;
 
 		const UText3DMaterialExtensionBase* MaterialExtension = TextComponent->GetMaterialExtension();
 
@@ -124,7 +124,14 @@ void UText3DDynamicMeshRenderer::OnUpdate(const UE::Text3D::Renderer::FUpdatePar
 		{
 			for (const int32 MaterialSlot : GlyphMeshData[InIndex].MaterialSlots)
 			{
-				UMaterialInterface* Material = MaterialExtension->GetMaterial(static_cast<EText3DGroupType>(MaterialSlot - (InIndex * SlotNames.Num())));
+				const FName StyleTag = InCharacter->GetStyleTag();
+				const EText3DGroupType GroupType = static_cast<EText3DGroupType>(MaterialSlot - (InIndex * SlotNames.Num()));
+
+				FMaterialParameters Parameters;
+				Parameters.Group = GroupType;
+				Parameters.Tag = StyleTag;
+
+				UMaterialInterface* Material = MaterialExtension->GetMaterial(Parameters);
 
 				if (Material != DynamicMeshComponent->GetMaterial(MaterialSlot))
 				{
@@ -296,7 +303,7 @@ void UText3DDynamicMeshRenderer::AppendGlyphMesh(FScopedMeshEditor& InMeshEditor
 				for (const TPair<int32, int32>& FromToPair : FromToMap)
 				{
 					InEditMesh.SetTriangleGroup(FromToPair.Value, InIndex);
-					const int32 MaterialID = (InIndex * UE::Text3D::Materials::SlotNames.Num()) + MaterialAttribute->GetValue(FromToPair.Value);
+					const int32 MaterialID = (InIndex * UE::Text3D::Material::SlotNames.Num()) + MaterialAttribute->GetValue(FromToPair.Value);
 					MaterialAttribute->SetValue(FromToPair.Value, MaterialID);
 					MaterialSlots.AddUnique(MaterialID);
 				}
diff --git a/Engine/Plugins/Experimental/Text3D/Source/Text3D/Private/Renderers/StaticMeshes/Text3DStaticMeshesRenderer.cpp b/Engine/Plugins/Experimental/Text3D/Source/Text3D/Private/Renderers/StaticMeshes/Text3DStaticMeshesRenderer.cpp
--- a/Engine/Plugins/Experimental/Text3D/Source/Text3D/Private/Renderers/StaticMeshes/Text3DStaticMeshesRenderer.cpp
+++ b/Engine/Plugins/Experimental/Text3D/Source/Text3D/Private/Renderers/StaticMeshes/Text3DStaticMeshesRenderer.cpp
@@ -91,7 +91,7 @@ void UText3DStaticMeshesRenderer::OnUpdate(const UE::Text3D::Renderer::FUpdatePa
 	}
 	else if (InParameters.CurrentFlag == EText3DRendererFlags::Material)
 	{
-		using namespace UE::Text3D::Materials;
+		using namespace UE::Text3D::Material;
 
 		const UText3DMaterialExtensionBase* MaterialExtension = TextComponent->GetMaterialExtension();
 
@@ -109,7 +109,14 @@ void UText3DStaticMeshesRenderer::OnUpdate(const UE::Text3D::Renderer::FUpdatePa
 						continue;
 					}
 
-					UMaterialInterface* Material = MaterialExtension->GetMaterial(static_cast<EText3DGroupType>(GroupIndex));
+					const EText3DGroupType GroupType = static_cast<EText3DGroupType>(GroupIndex);
+					const FName StyleTag = InCharacter->GetStyleTag();
+
+					FMaterialParameters Parameters;
+					Parameters.Group = GroupType;
+					Parameters.Tag = StyleTag;
+
+					UMaterialInterface* Material = MaterialExtension->GetMaterial(Parameters);
 
 					if (Material != StaticMeshComponent->GetMaterial(MaterialIndex))
 					{
@@ -239,15 +246,17 @@ void UText3DStaticMeshesRenderer::AllocateComponents(int32 InCount)
 			UStaticMeshComponent* MeshComponent = nullptr;
 			if (!CharacterMeshesPool.IsEmpty())
 			{
-				MeshComponent = CharacterMeshes.Add_GetRef(CharacterMeshesPool.Pop());
+				MeshComponent = CharacterMeshesPool.Pop();
 			}
-			else
+
+			if (!MeshComponent)
 			{
 				const FName StaticMeshComponentName = MakeUniqueObjectName(this, UStaticMeshComponent::StaticClass(), FName(TEXT("CharacterMesh")));
 				MeshComponent = NewObject<UStaticMeshComponent>(this, StaticMeshComponentName);
-				CharacterMeshes.Add(MeshComponent);
 			}
 
+			CharacterMeshes.Add(MeshComponent);
+
 			Actor->AddOwnedComponent(MeshComponent);
 			MeshComponent->RegisterComponent();
 			MeshComponent->AttachToComponent(TextRoot, FAttachmentTransformRules::KeepRelativeTransform);
diff --git a/Engine/Plugins/Experimental/Text3D/Source/Text3D/Private/Renderers/Text3DRendererBase.cpp b/Engine/Plugins/Experimental/Text3D/Source/Text3D/Private/Renderers/Text3DRendererBase.cpp
--- a/Engine/Plugins/Experimental/Text3D/Source/Text3D/Private/Renderers/Text3DRendererBase.cpp
+++ b/Engine/Plugins/Experimental/Text3D/Source/Text3D/Private/Renderers/Text3DRendererBase.cpp
@@ -3,13 +3,13 @@
 #include "Renderers/Text3DRendererBase.h"
 
 #include "Extensions/Text3DCharacterExtensionBase.h"
-#include "Extensions/Text3DFormatStyleExtensionBase.h"
 #include "Extensions/Text3DGeometryExtensionBase.h"
 #include "Extensions/Text3DLayoutEffectBase.h"
 #include "Extensions/Text3DLayoutExtensionBase.h"
 #include "Extensions/Text3DMaterialExtensionBase.h"
 #include "Extensions/Text3DRenderingExtensionBase.h"
-#include "Extensions/Text3DTemplateExtensionBase.h"
+#include "Extensions/Text3DStyleExtensionBase.h"
+#include "Extensions/Text3DTokenExtensionBase.h"
 #include "GameFramework/Actor.h"
 #include "Logs/Text3DLogs.h"
 #include "Settings/Text3DProjectSettings.h"
@@ -109,8 +109,8 @@ void UText3DRendererBase::Update(EText3DRendererFlags InFlags)
 
 	TArray<UText3DExtensionBase*, TInlineAllocator<6>> Extensions
 	{
-		Text3DComponent->GetTemplateExtension(),
-		Text3DComponent->GetFormatStyleExtension(),
+		Text3DComponent->GetTokenExtension(),
+		Text3DComponent->GetStyleExtension(),
 		Text3DComponent->GetCharacterExtension(),
 		Text3DComponent->GetGeometryExtension(),
 		Text3DComponent->GetLayoutExtension(),
diff --git a/Engine/Plugins/Experimental/Text3D/Source/Text3D/Private/Settings/Text3DProjectSettings.cpp b/Engine/Plugins/Experimental/Text3D/Source/Text3D/Private/Settings/Text3DProjectSettings.cpp
--- a/Engine/Plugins/Experimental/Text3D/Source/Text3D/Private/Settings/Text3DProjectSettings.cpp
+++ b/Engine/Plugins/Experimental/Text3D/Source/Text3D/Private/Settings/Text3DProjectSettings.cpp
@@ -6,6 +6,7 @@
 #include "Engine/FontFace.h"
 #include "Materials/Material.h"
 #include "Modules/ModuleManager.h"
+#include "Styles/Text3DStyleSet.h"
 #include "Subsystems/Text3DEngineSubsystem.h"
 #include "UObject/SoftObjectPath.h"
 
@@ -51,6 +52,11 @@ UMaterial* UText3DProjectSettings::GetDefaultMaterial() const
 	return DefaultMaterial.LoadSynchronous();
 }
 
+UText3DStyleSet* UText3DProjectSettings::GetDefaultStyleSet() const
+{
+	return DefaultStyleSet.LoadSynchronous();
+}
+
 UFont* UText3DProjectSettings::GetFallbackFont() const
 {
 	return FallbackFont.LoadSynchronous();
diff --git a/Engine/Plugins/Experimental/Text3D/Source/Text3D/Private/Styles/Text3DFormatStyleBase.cpp b/Engine/Plugins/Experimental/Text3D/Source/Text3D/Private/Styles/Text3DFormatStyleBase.cpp
deleted file mode 100644
diff --git a/Engine/Plugins/Experimental/Text3D/Source/Text3D/Private/Styles/Text3DStyleBase.cpp b/Engine/Plugins/Experimental/Text3D/Source/Text3D/Private/Styles/Text3DStyleBase.cpp
new file mode 100644
--- /dev/null
+++ b/Engine/Plugins/Experimental/Text3D/Source/Text3D/Private/Styles/Text3DStyleBase.cpp
@@ -0,0 +1,176 @@
+// Copyright Epic Games, Inc. All Rights Reserved.
+
+#include "Styles/Text3DStyleBase.h"
+
+#include "Engine/Font.h"
+#include "Settings/Text3DProjectSettings.h"
+#include "Styling/StyleDefaults.h"
+#include "Text3DComponent.h"
+
+#if WITH_EDITOR
+void UText3DStyleBase::PostEditChangeProperty(FPropertyChangedEvent& InPropertyChangedEvent)
+{
+	Super::PostEditChangeProperty(InPropertyChangedEvent);
+
+	const FName MemberProperty = InPropertyChangedEvent.GetMemberPropertyName();
+
+	if (MemberProperty == GET_MEMBER_NAME_CHECKED(UText3DStyleBase, Font))
+	{
+		OnFontChanged();
+	}
+	else if (MemberProperty == GET_MEMBER_NAME_CHECKED(UText3DStyleBase, StyleName)
+		|| MemberProperty == GET_MEMBER_NAME_CHECKED(UText3DStyleBase, bOverrideFont)
+		|| MemberProperty == GET_MEMBER_NAME_CHECKED(UText3DStyleBase, FontTypeface)
+		|| MemberProperty == GET_MEMBER_NAME_CHECKED(UText3DStyleBase, bOverrideFrontColor)
+		|| MemberProperty == GET_MEMBER_NAME_CHECKED(UText3DStyleBase, FrontColor))
+	{
+		OnStylePropertiesChanged();
+	}
+}
+#endif
+
+UText3DStyleBase::UText3DStyleBase()
+{
+	if (const UText3DProjectSettings* Text3DSettings = UText3DProjectSettings::Get())
+	{
+		Font = Text3DSettings->GetFallbackFont();
+		OnFontChanged();
+	}
+}
+
+void UText3DStyleBase::SetStyleName(FName InName)
+{
+	if (InName == StyleName)
+	{
+		return;
+	}
+
+	StyleName = InName;
+	OnStylePropertiesChanged();
+}
+
+void UText3DStyleBase::SetOverrideFont(bool bInOverride)
+{
+	if (bOverrideFont == bInOverride)
+	{
+		return;
+	}
+
+	bOverrideFont = bInOverride;
+	OnStylePropertiesChanged();
+}
+
+void UText3DStyleBase::SetFont(UFont* InFont)
+{
+	if (InFont == Font)
+	{
+		return;
+	}
+
+	Font = InFont;
+	OnFontChanged();
+}
+
+void UText3DStyleBase::SetFontTypeface(FName InTypeface)
+{
+	if (InTypeface == FontTypeface)
+	{
+		return;
+	}
+
+	if (!GetTypefaceNames().Contains(InTypeface))
+	{
+		return;
+	}
+
+	FontTypeface = InTypeface;
+	OnStylePropertiesChanged();
+}
+
+void UText3DStyleBase::SetOverrideFrontColor(bool bInOverride)
+{
+	if (bOverrideFrontColor == bInOverride)
+	{
+		return;
+	}
+
+	bOverrideFrontColor = bInOverride;
+	OnStylePropertiesChanged();
+}
+
+void UText3DStyleBase::SetFrontColor(FLinearColor InColor)
+{
+	if (FrontColor.Equals(InColor))
+	{
+		return;
+	}
+
+	FrontColor = InColor;
+	OnStylePropertiesChanged();
+}
+
+void UText3DStyleBase::ConfigureStyle(FTextBlockStyle& OutStyle) const
+{
+	if (bOverrideFont)
+	{
+		FSlateFontInfo FontInfo(Font, UE::Text3D::Metrics::SlateFontSize);
+		FontInfo.CompositeFont = Font && Font->GetCompositeFont() ? MakeShared<FCompositeFont>(*Font->GetCompositeFont()) : FStyleDefaults::GetFontInfo().CompositeFont;
+		FontInfo.TypefaceFontName = FontTypeface;
+		FontInfo.FontFallback = EFontFallback::FF_NoFallback;
+		OutStyle.SetFont(FontInfo);
+	}
+
+	if (bOverrideFrontColor)
+	{
+		OutStyle.SetColorAndOpacity(FrontColor);
+	}
+}
+
+void UText3DStyleBase::OnStylePropertiesChanged()
+{
+	if (!StyleName.IsNone())
+	{
+		if (UText3DComponent* Text3DComponent = GetTypedOuter<UText3DComponent>())
+		{
+			Text3DComponent->RequestUpdate(EText3DRendererFlags::All, /** Immediate */false);
+		}
+	}
+}
+
+void UText3DStyleBase::OnFontChanged()
+{
+	const TArray<FName> Typefaces = GetTypefaceNames();
+	if (!Typefaces.Contains(FontTypeface) && !Typefaces.IsEmpty())
+	{
+		FontTypeface = Typefaces[0];
+	}
+	OnStylePropertiesChanged();
+}
+
+TArray<FName> UText3DStyleBase::GetTypefaceNames() const
+{
+	TArray<FName> TypefaceNames;
+
+	if (Font)
+	{
+		for (const FTypefaceEntry& TypeFaceFont : GetAvailableTypefaces())
+		{
+			TypefaceNames.Add(TypeFaceFont.Name);
+		}
+	}
+
+	return TypefaceNames;
+}
+
+TArrayView<const FTypefaceEntry> UText3DStyleBase::GetAvailableTypefaces() const
+{
+	if (Font)
+	{
+		if (const FCompositeFont* CompositeFont = Font->GetCompositeFont())
+		{
+			return CompositeFont->DefaultTypeface.Fonts;
+		}
+	}
+
+	return {};
+}
diff --git a/Engine/Plugins/Experimental/Text3D/Source/Text3D/Private/Styles/Text3DStyleSet.cpp b/Engine/Plugins/Experimental/Text3D/Source/Text3D/Private/Styles/Text3DStyleSet.cpp
new file mode 100644
--- /dev/null
+++ b/Engine/Plugins/Experimental/Text3D/Source/Text3D/Private/Styles/Text3DStyleSet.cpp
@@ -0,0 +1,24 @@
+// Copyright Epic Games, Inc. All Rights Reserved.
+
+#include "Styles/Text3DStyleSet.h"
+
+UText3DStyleSet::FOnStyleSetUpdated UText3DStyleSet::OnStyleSetUpdatedDelegate;
+
+#if WITH_EDITOR
+void UText3DStyleSet::PostEditChangeProperty(FPropertyChangedEvent& InPropertyChangedEvent)
+{
+	Super::PostEditChangeProperty(InPropertyChangedEvent);
+
+	const FName MemberName = InPropertyChangedEvent.GetMemberPropertyName();
+
+	if (MemberName == GET_MEMBER_NAME_CHECKED(UText3DStyleSet, Styles))
+	{
+		OnStyleSetPropertiesChanged();
+	}
+}
+#endif
+
+void UText3DStyleSet::OnStyleSetPropertiesChanged()
+{
+	OnStyleSetUpdatedDelegate.Broadcast(this);
+}
diff --git a/Engine/Plugins/Experimental/Text3D/Source/Text3D/Private/Subsystems/Text3DEngineSubsystem.cpp b/Engine/Plugins/Experimental/Text3D/Source/Text3D/Private/Subsystems/Text3DEngineSubsystem.cpp
--- a/Engine/Plugins/Experimental/Text3D/Source/Text3D/Private/Subsystems/Text3DEngineSubsystem.cpp
+++ b/Engine/Plugins/Experimental/Text3D/Source/Text3D/Private/Subsystems/Text3DEngineSubsystem.cpp
@@ -166,10 +166,13 @@ uint32 FText3DFontFaceCache::GetFontFaceHash(const TSharedPtr<FFreeTypeFace>& In
 {
 	if (const FT_Face Face = InFontFace->GetFace())
 	{
-		const FString FaceFamily(Face->family_name);
-		const FString FaceStyle(Face->style_name);
-		const int32 FaceIndex = Face->face_index;
-		return HashCombine(HashCombine(GetTypeHash(FaceFamily), GetTypeHash(FaceStyle)), FaceIndex);
+		if (Face->family_name != nullptr && Face->style_name != nullptr)
+		{
+			const FString FaceFamily(Face->family_name);
+			const FString FaceStyle(Face->style_name);
+			const int32 FaceIndex = Face->face_index;
+			return HashCombine(HashCombine(GetTypeHash(FaceFamily), GetTypeHash(FaceStyle)), FaceIndex);
+		}
 	}
 
 	return 0;
diff --git a/Engine/Plugins/Experimental/Text3D/Source/Text3D/Private/Templates/Text3DTemplateBase.cpp b/Engine/Plugins/Experimental/Text3D/Source/Text3D/Private/Templates/Text3DTemplateBase.cpp
deleted file mode 100644
diff --git a/Engine/Plugins/Experimental/Text3D/Source/Text3D/Private/Text3DComponent.cpp b/Engine/Plugins/Experimental/Text3D/Source/Text3D/Private/Text3DComponent.cpp
--- a/Engine/Plugins/Experimental/Text3D/Source/Text3D/Private/Text3DComponent.cpp
+++ b/Engine/Plugins/Experimental/Text3D/Source/Text3D/Private/Text3DComponent.cpp
@@ -6,13 +6,13 @@
 #include "Engine/Engine.h"
 #include "Engine/Font.h"
 #include "Extensions/Text3DDefaultCharacterExtension.h"
-#include "Extensions/Text3DDefaultFormatStyleExtension.h"
 #include "Extensions/Text3DDefaultGeometryExtension.h"
 #include "Extensions/Text3DDefaultLayoutExtension.h"
 #include "Extensions/Text3DDefaultMaterialExtension.h"
 #include "Extensions/Text3DDefaultRenderingExtension.h"
+#include "Extensions/Text3DDefaultStyleExtension.h"
 #include "Extensions/Text3DLayoutEffectBase.h"
-#include "Extensions/Text3DDefaultTemplateExtension.h"
+#include "Extensions/Text3DDefaultTokenExtension.h"
 #include "Fonts/CompositeFont.h"
 #include "Internationalization/Regex.h"
 #include "Logs/Text3DLogs.h"
@@ -23,7 +23,6 @@
 #include "Renderers/StaticMeshes/Text3DStaticMeshesRenderer.h"
 #include "Renderers/Text3DRendererBase.h"
 #include "Settings/Text3DProjectSettings.h"
-#include "Subsystems/Text3DEngineSubsystem.h"
 #include "Text3DComponentVersion.h"
 #include "UObject/ConstructorHelpers.h"
 #include "UObject/UObjectGlobals.h"
@@ -50,8 +49,8 @@ UText3DComponent::UText3DComponent()
 	TextRendererClass = UText3DStaticMeshesRenderer::StaticClass();
 
 	// Extensions
-	TemplateExtension = CreateDefaultSubobject<UText3DTemplateExtensionBase, UText3DDefaultTemplateExtension>(TEXT("TemplateExtension"));
-	FormatStyleExtension = CreateDefaultSubobject<UText3DFormatStyleExtensionBase, UText3DDefaultFormatStyleExtension>(TEXT("FormatStyleExtension"));
+	TokenExtension = CreateDefaultSubobject<UText3DTokenExtensionBase, UText3DDefaultTokenExtension>(TEXT("TemplateExtension"));
+	StyleExtension = CreateDefaultSubobject<UText3DStyleExtensionBase, UText3DDefaultStyleExtension>(TEXT("StyleExtension"));
 	MaterialExtension = CreateDefaultSubobject<UText3DMaterialExtensionBase, UText3DDefaultMaterialExtension>(TEXT("MaterialExtension"));
 	GeometryExtension = CreateDefaultSubobject<UText3DGeometryExtensionBase, UText3DDefaultGeometryExtension>(TEXT("GeometryExtension"));
 	LayoutExtension = CreateDefaultSubobject<UText3DLayoutExtensionBase, UText3DDefaultLayoutExtension>(TEXT("LayoutExtension"));
@@ -100,11 +99,6 @@ void UText3DComponent::RequestUpdate(EText3DRendererFlags InFlags, bool bInImmed
 {
 	if (InFlags != EText3DRendererFlags::None)
 	{
-		if (IsTemplate(RF_ClassDefaultObject))
-		{
-			return;
-		}
-
 		if (bInImmediate)
 		{
 			Rebuild(InFlags);
@@ -442,16 +436,24 @@ void UText3DComponent::ForceUpdateText()
 	SetText(FText::GetEmpty());
 	SetText(PrevText);
 }
+
+void UText3DComponent::OpenProjectSettings()
+{
+	if (const UText3DProjectSettings* TextSettings = UText3DProjectSettings::Get())
+	{
+		TextSettings->OpenEditorSettingsWindow();
+	}
+}
 #endif
 
 const FText& UText3DComponent::GetFormattedText() const
 {
-	if (const UText3DTemplateExtensionBase* Extension = GetTemplateExtension())
+	if (const UText3DTokenExtensionBase* Extension = GetTokenExtension())
 	{
 		return Extension->GetFormattedText();
 	}
 
-	return FText::GetEmpty();
+	return Text;
 }
 
 void UText3DComponent::SetText(const FText& Value)
@@ -539,7 +541,9 @@ UMaterialInterface* UText3DComponent::GetBackMaterial() const
 {
 	if (MaterialExtension)
 	{
-		return MaterialExtension->GetMaterial(EText3DGroupType::Back);
+		UE::Text3D::Material::FMaterialParameters Parameters;
+		Parameters.Group = EText3DGroupType::Back;
+		return MaterialExtension->GetMaterial(Parameters);
 	}
 
 	return nullptr;
@@ -549,7 +553,9 @@ UMaterialInterface* UText3DComponent::GetExtrudeMaterial() const
 {
 	if (MaterialExtension)
 	{
-		return MaterialExtension->GetMaterial(EText3DGroupType::Extrude);
+		UE::Text3D::Material::FMaterialParameters Parameters;
+		Parameters.Group = EText3DGroupType::Extrude;
+		return MaterialExtension->GetMaterial(Parameters);
 	}
 
 	return nullptr;
@@ -559,7 +565,9 @@ UMaterialInterface* UText3DComponent::GetBevelMaterial() const
 {
 	if (MaterialExtension)
 	{
-		return MaterialExtension->GetMaterial(EText3DGroupType::Bevel);
+		UE::Text3D::Material::FMaterialParameters Parameters;
+		Parameters.Group = EText3DGroupType::Bevel;
+		return MaterialExtension->GetMaterial(Parameters);
 	}
 
 	return nullptr;
@@ -569,7 +577,9 @@ UMaterialInterface* UText3DComponent::GetFrontMaterial() const
 {
 	if (MaterialExtension)
 	{
-		return MaterialExtension->GetMaterial(EText3DGroupType::Front);
+		UE::Text3D::Material::FMaterialParameters Parameters;
+		Parameters.Group = EText3DGroupType::Front;
+		return MaterialExtension->GetMaterial(Parameters);
 	}
 
 	return nullptr;
@@ -579,7 +589,9 @@ void UText3DComponent::SetFrontMaterial(UMaterialInterface* Value)
 {
 	if (MaterialExtension)
 	{
-		return MaterialExtension->SetMaterial(EText3DGroupType::Front, Value);
+		UE::Text3D::Material::FMaterialParameters Parameters;
+		Parameters.Group = EText3DGroupType::Front;
+		return MaterialExtension->SetMaterial(Parameters, Value);
 	}
 }
 
@@ -587,7 +599,9 @@ void UText3DComponent::SetBevelMaterial(UMaterialInterface* Value)
 {
 	if (MaterialExtension)
 	{
-		return MaterialExtension->SetMaterial(EText3DGroupType::Bevel, Value);
+		UE::Text3D::Material::FMaterialParameters Parameters;
+		Parameters.Group = EText3DGroupType::Bevel;
+		return MaterialExtension->SetMaterial(Parameters, Value);
 	}
 }
 
@@ -595,7 +609,9 @@ void UText3DComponent::SetExtrudeMaterial(UMaterialInterface* Value)
 {
 	if (MaterialExtension)
 	{
-		return MaterialExtension->SetMaterial(EText3DGroupType::Extrude, Value);
+		UE::Text3D::Material::FMaterialParameters Parameters;
+		Parameters.Group = EText3DGroupType::Extrude;
+		return MaterialExtension->SetMaterial(Parameters, Value);
 	}
 }
 
@@ -603,7 +619,9 @@ void UText3DComponent::SetBackMaterial(UMaterialInterface* Value)
 {
 	if (MaterialExtension)
 	{
-		return MaterialExtension->SetMaterial(EText3DGroupType::Back, Value);
+		UE::Text3D::Material::FMaterialParameters Parameters;
+		Parameters.Group = EText3DGroupType::Back;
+		return MaterialExtension->SetMaterial(Parameters, Value);
 	}
 }
 
@@ -622,24 +640,24 @@ UText3DRendererBase* UText3DComponent::GetTextRenderer() const
 	return TextRenderer;
 }
 
-UText3DTemplateExtensionBase* UText3DComponent::GetTemplateExtension(TSubclassOf<UText3DTemplateExtensionBase> InExtensionClass)
+UText3DTokenExtensionBase* UText3DComponent::GetTokenExtension(TSubclassOf<UText3DTokenExtensionBase> InExtensionClass)
 {
-	if (!TemplateExtension || !InExtensionClass.Get())
+	if (!TokenExtension || !InExtensionClass.Get())
 	{
 		return nullptr;
 	}
 
-	return TemplateExtension->IsA(InExtensionClass.Get()) ? TemplateExtension : nullptr;
+	return TokenExtension->IsA(InExtensionClass.Get()) ? TokenExtension : nullptr;
 }
 
-UText3DFormatStyleExtensionBase* UText3DComponent::GetFormatStyleExtension(TSubclassOf<UText3DFormatStyleExtensionBase> InExtensionClass)
+UText3DStyleExtensionBase* UText3DComponent::GetStyleExtension(TSubclassOf<UText3DStyleExtensionBase> InExtensionClass)
 {
-	if (!FormatStyleExtension || !InExtensionClass.Get())
+	if (!StyleExtension || !InExtensionClass.Get())
 	{
 		return nullptr;
 	}
 
-	return FormatStyleExtension->IsA(InExtensionClass.Get()) ? FormatStyleExtension : nullptr;
+	return StyleExtension->IsA(InExtensionClass.Get()) ? StyleExtension : nullptr;
 }
 
 UText3DLayoutExtensionBase* UText3DComponent::GetLayoutExtension(TSubclassOf<UText3DLayoutExtensionBase> InExtensionClass)
@@ -1023,7 +1041,7 @@ void UText3DComponent::Rebuild(EText3DRendererFlags InUpdateFlags)
 	{
 		if (!EnumHasAllFlags(UpdateFlags, InUpdateFlags))
 		{
-			UE_LOG(LogText3D, Warning, TEXT("[%s] Rebuild called with update flags (%u) while package is being saved. This new Rebuild will be deferred to the next tick."), *GetReadableName(), InUpdateFlags);
+			UE_LOG(LogText3D, Log, TEXT("[%s] Rebuild called with update flags (%u) while package is being saved. This new Rebuild will be deferred to the next tick."), *GetReadableName(), InUpdateFlags);
 		}
 
 		QueueRebuild(InUpdateFlags);
@@ -1034,13 +1052,15 @@ void UText3DComponent::Rebuild(EText3DRendererFlags InUpdateFlags)
 	{
 		if (!EnumHasAllFlags(UpdateFlags, InUpdateFlags))
 		{
-			UE_LOG(LogText3D, Warning, TEXT("[%s] Rebuild called with update flags (%u) while building text. This new Rebuild will be deferred to the next tick."), *GetReadableName(), InUpdateFlags);
+			UE_LOG(LogText3D, Log, TEXT("[%s] Rebuild called with update flags (%u) while building text. This new Rebuild will be deferred to the next tick."), *GetReadableName(), InUpdateFlags);
 		}
 
 		QueueRebuild(InUpdateFlags);
 		return;
 	}
 
+	EnumRemoveFlags(UpdateFlags, InUpdateFlags);
+
 	bIsUpdatingText = true;
 
 	ON_SCOPE_EXIT
diff --git a/Engine/Plugins/Experimental/Text3D/Source/Text3D/Private/Text3DInternalTypes.h b/Engine/Plugins/Experimental/Text3D/Source/Text3D/Private/Text3DInternalTypes.h
--- a/Engine/Plugins/Experimental/Text3D/Source/Text3D/Private/Text3DInternalTypes.h
+++ b/Engine/Plugins/Experimental/Text3D/Source/Text3D/Private/Text3DInternalTypes.h
@@ -5,7 +5,7 @@
 #include "Containers/Array.h"
 #include "UObject/NameTypes.h"
 
-namespace UE::Text3D::Materials
+namespace UE::Text3D::Material
 {
 	static const TArray<FLazyName> SlotNames =
 	{
diff --git a/Engine/Plugins/Experimental/Text3D/Source/Text3D/Private/Tokens/Text3DTokenBase.cpp b/Engine/Plugins/Experimental/Text3D/Source/Text3D/Private/Tokens/Text3DTokenBase.cpp
new file mode 100644
--- /dev/null
+++ b/Engine/Plugins/Experimental/Text3D/Source/Text3D/Private/Tokens/Text3DTokenBase.cpp
@@ -0,0 +1,67 @@
+// Copyright Epic Games, Inc. All Rights Reserved.
+
+#include "Tokens/Text3DTokenBase.h"
+
+#include "Text3DComponent.h"
+
+#if WITH_EDITOR
+void UText3DTokenBase::PostEditChangeProperty(FPropertyChangedEvent& InPropertyChangedEvent)
+{
+	Super::PostEditChangeProperty(InPropertyChangedEvent);
+
+	const FName MemberName = InPropertyChangedEvent.GetMemberPropertyName();
+
+	if (MemberName == GET_MEMBER_NAME_CHECKED(UText3DTokenBase, TokenName)
+		|| MemberName == GET_MEMBER_NAME_CHECKED(UText3DTokenBase, Content))
+	{
+		OnTokenPropertiesChanged();
+	}
+}
+#endif
+
+void UText3DTokenBase::SetTokenName(FName InName)
+{
+	if (InName == TokenName)
+	{
+		return;
+	}
+
+	TokenName = InName;
+	OnTokenPropertiesChanged();
+}
+
+void UText3DTokenBase::SetContent(const FText& InContent)
+{
+	if (InContent.EqualTo(Content))
+	{
+		return;
+	}
+
+	Content = InContent;
+	OnTokenPropertiesChanged();
+}
+
+void UText3DTokenBase::OnTokenPropertiesChanged()
+{
+	if (!TokenName.IsNone())
+	{
+		if (UText3DComponent* Text3DComponent = GetTypedOuter<UText3DComponent>())
+		{
+			Text3DComponent->RequestUpdate(EText3DRendererFlags::All, /** Immediate */false);
+		}
+	}
+}
+
+void UText3DTokenBase::CollectTokens(FFormatNamedArguments& InNamedArguments)
+{
+	if (!TokenName.IsNone())
+	{
+		const FString TokenPattern = TokenName.ToString();
+		const FString EscapeTokenPattern = FString::Printf(TEXT("\\%s\\"), *TokenPattern);
+		const FString EscapeTokenMarker = FString::Printf(TEXT("{%s}"), *TokenPattern);
+
+		// Replaces {\token\} by {token} and {token} by its content
+		InNamedArguments.Add(*EscapeTokenPattern, FText::FromString(EscapeTokenMarker));
+		InNamedArguments.Add(*TokenPattern, Content);
+	}
+}
diff --git a/Engine/Plugins/Experimental/Text3D/Source/Text3D/Public/Characters/Text3DCharacterBase.h b/Engine/Plugins/Experimental/Text3D/Source/Text3D/Public/Characters/Text3DCharacterBase.h
--- a/Engine/Plugins/Experimental/Text3D/Source/Text3D/Public/Characters/Text3DCharacterBase.h
+++ b/Engine/Plugins/Experimental/Text3D/Source/Text3D/Public/Characters/Text3DCharacterBase.h
@@ -35,6 +35,12 @@ public:
 	const UE::Text3D::Geometry::FCachedFontFaceGlyphHandle& GetFontFaceGlyphHandle() const;
 	void SetFontFaceGlyphHandle(const UE::Text3D::Geometry::FCachedFontFaceGlyphHandle& InHandle);
 
+	void SetStyleTag(FName InStyle);
+	FName GetStyleTag() const
+	{
+		return StyleTag;
+	}
+
 	uint32 GetGlyphIndex() const;
 	const FText3DCachedMesh* GetGlyphMesh() const;
 	void GetGlyphMeshBoundsAndOffset(FBox& OutBounds, FVector& OutOffset) const;
@@ -109,4 +115,7 @@ protected:
 
 	/** Handle to the current glyph mesh this character represents */
 	UE::Text3D::Geometry::FCachedFontFaceGlyphHandle FontFaceGlyphHandle;
+
+	/** Tag used to identify the custom style applied to this character */
+	FName StyleTag = NAME_None;
 };
diff --git a/Engine/Plugins/Experimental/Text3D/Source/Text3D/Public/Extensions/Text3DDefaultFormatStyleExtension.h b/Engine/Plugins/Experimental/Text3D/Source/Text3D/Public/Extensions/Text3DDefaultFormatStyleExtension.h
deleted file mode 100644
diff --git a/Engine/Plugins/Experimental/Text3D/Source/Text3D/Public/Extensions/Text3DDefaultLayoutExtension.h b/Engine/Plugins/Experimental/Text3D/Source/Text3D/Public/Extensions/Text3DDefaultLayoutExtension.h
--- a/Engine/Plugins/Experimental/Text3D/Source/Text3D/Public/Extensions/Text3DDefaultLayoutExtension.h
+++ b/Engine/Plugins/Experimental/Text3D/Source/Text3D/Public/Extensions/Text3DDefaultLayoutExtension.h
@@ -8,6 +8,11 @@
 class FRichTextLayoutMarshaller;
 class FText3DLayout;
 
+namespace UE::Text3D::Layout
+{
+	struct FGlyphText;
+}
+
 UCLASS(MinimalAPI)
 class UText3DDefaultLayoutExtension : public UText3DLayoutExtensionBase
 {
@@ -144,6 +149,7 @@ public:
 protected:
 	//~ Begin UObject
 #if WITH_EDITOR
+	virtual void PostEditUndo() override;
 	virtual void PostEditChangeProperty(FPropertyChangedEvent& InEvent) override;
 #endif
 	//~ End UObject
@@ -211,7 +217,7 @@ private:
 	FVector TextScale = FVector::ZeroVector;
 
 	/** Caches the last result of ShapedText, to allow faster updates of layout changes. */
-	TSharedPtr<FText3DShapedGlyphText> ShapedText;
+	TSharedPtr<UE::Text3D::Layout::FGlyphText> GlyphText;
 
 	/** Stores the text layout calculated by the TextLayoutMarshaller. */
 	TSharedPtr<FText3DLayout> TextLayout;
diff --git a/Engine/Plugins/Experimental/Text3D/Source/Text3D/Public/Extensions/Text3DDefaultMaterialExtension.h b/Engine/Plugins/Experimental/Text3D/Source/Text3D/Public/Extensions/Text3DDefaultMaterialExtension.h
--- a/Engine/Plugins/Experimental/Text3D/Source/Text3D/Public/Extensions/Text3DDefaultMaterialExtension.h
+++ b/Engine/Plugins/Experimental/Text3D/Source/Text3D/Public/Extensions/Text3DDefaultMaterialExtension.h
@@ -240,6 +240,7 @@ protected:
 	virtual void PostLoad() override;
 #if WITH_EDITOR
 	virtual void PostEditChangeProperty(FPropertyChangedEvent& InEvent) override;
+	virtual void PostEditUndo() override;
 #endif
 	//~ End UObject
 
@@ -249,11 +250,17 @@ protected:
 	//~ End UText3DExtensionBase
 
 	//~ Begin UText3DMaterialExtensionBase
-	virtual void SetMaterial(EText3DGroupType InGroup, UMaterialInterface* InMaterial) override;
-	virtual UMaterialInterface* GetMaterial(EText3DGroupType InGroup) const override;
+	virtual void SetMaterial(const UE::Text3D::Material::FMaterialParameters& InParameters, UMaterialInterface* InMaterial) override;
+	virtual UMaterialInterface* GetMaterial(const UE::Text3D::Material::FMaterialParameters& InParameters) const override;
+	virtual void ClearMaterialOverrides() override;
+	virtual void RegisterMaterialOverride(FName InTag) override;
+	virtual void ForEachMaterial(TFunctionRef<bool(const UE::Text3D::Material::FMaterialParameters&, UMaterialInterface*)> InFunctor) const override;
+	virtual int32 GetMaterialCount() const override;
+	virtual void GetMaterialNames(TArray<FName>& OutNames) const override;
 	//~ End UText3DMaterialExtensionBase
 
-	UMaterialInstanceDynamic* FindOrAdd(EText3DGroupType InGroup);
+	UMaterialInstanceDynamic* FindOrAdd(EText3DGroupType InGroup, FName InTag = NAME_None);
+	UMaterialInstanceDynamic* Find(EText3DGroupType InGroup, FName InTag = NAME_None) const;
 	void SetVectorParameter(TConstArrayView<UMaterialInstanceDynamic*> InMaterials, FName InKey, FVector InValue);
 	void SetVectorParameter(TConstArrayView<UMaterialInstanceDynamic*> InMaterials, FName InKey, FLinearColor InValue);
 	void SetScalarParameter(TConstArrayView<UMaterialInstanceDynamic*> InMaterials, FName InKey, float InValue);
@@ -339,6 +346,10 @@ protected:
 	UPROPERTY(EditAnywhere, Getter, Setter, Category = "Material", meta = (EditCondition="!bUseSingleMaterial && Style == EText3DMaterialStyle::Custom", EditConditionHides, AllowPrivateAccess = "true"))
 	TObjectPtr<UMaterialInterface> BackMaterial;
 
+	/** Additional overrides defined for specific group */
+	UPROPERTY(Transient, DuplicateTransient)
+	TArray<FText3DMaterialOverride> MaterialOverrides;
+
 	/** Cached group dynamic materials created during session based on their options */
 	UPROPERTY(Transient, DuplicateTransient)
 	TMap<FText3DMaterialGroupKey, TObjectPtr<UMaterialInstanceDynamic>> GroupDynamicMaterials;
diff --git a/Engine/Plugins/Experimental/Text3D/Source/Text3D/Public/Extensions/Text3DDefaultStyleExtension.h b/Engine/Plugins/Experimental/Text3D/Source/Text3D/Public/Extensions/Text3DDefaultStyleExtension.h
new file mode 100644
--- /dev/null
+++ b/Engine/Plugins/Experimental/Text3D/Source/Text3D/Public/Extensions/Text3DDefaultStyleExtension.h
@@ -0,0 +1,67 @@
+// Copyright Epic Games, Inc. All Rights Reserved.
+
+#pragma once
+
+#include "Containers/Set.h"
+#include "Text3DStyleExtensionBase.h"
+#include "Text3DDefaultStyleExtension.generated.h"
+
+class FSlateStyleSet;
+class UText3DStyleBase;
+class UText3DStyleSet;
+struct FTextBlockStyle;
+
+UCLASS(MinimalAPI)
+class UText3DDefaultStyleExtension : public UText3DStyleExtensionBase
+{
+	GENERATED_BODY()
+
+public:
+	UText3DDefaultStyleExtension();
+
+	UFUNCTION(BlueprintCallable, Category = "Text3D|Style")
+	void SetStyleSet(UText3DStyleSet* InStyleSet);
+
+	UFUNCTION(BlueprintPure, Category = "Text3D|Style")
+	UText3DStyleSet* GetStyleSet() const
+	{
+		return StyleSet;
+	}
+
+protected:
+	//~ Begin UObject
+#if WITH_EDITOR
+	virtual void PostEditUndo() override;
+	virtual void PostEditChangeProperty(FPropertyChangedEvent& InPropertyChangedEvent) override;
+#endif
+	//~ End UObject
+
+	//~ Begin UText3DFormatStyleExtensionBase
+	virtual TSharedPtr<FTextBlockStyle> GetDefaultStyle() const override;
+    virtual TSharedPtr<FSlateStyleSet> GetCustomStyles() const override;
+	virtual UText3DStyleBase* GetStyle(FName InStyle) const override;
+    //~ End UText3DFormatStyleExtensionBase
+    
+    //~ Begin UText3DExtensionBase
+    virtual EText3DExtensionResult PreRendererUpdate(const UE::Text3D::Renderer::FUpdateParameters& InParameters) override;
+    virtual EText3DExtensionResult PostRendererUpdate(const UE::Text3D::Renderer::FUpdateParameters& InParameters) override;
+    //~ End UText3DExtensionBase
+
+	void OnStyleSetUpdated(UText3DStyleSet* InStyleSet);
+	void OnStyleSetChanged();
+
+	/** A style set asset that defines custom style rules */
+	UPROPERTY(EditAnywhere, Setter, Getter, Category="Style", meta=(AllowPrivateAccess="true"))
+	TObjectPtr<UText3DStyleSet> StyleSet;
+
+	/**
+	 * Define your styles and use them in text like this <MyStyle>...</>,
+	 * the text between start and end tag will have this style applied
+	 */
+	UPROPERTY(EditInstanceOnly, Instanced, Category="Style", NoClear, meta=(AllowPrivateAccess="true"))
+	TArray<TObjectPtr<UText3DStyleBase>> Styles;
+
+private:
+	TSharedPtr<FTextBlockStyle> DefaultFontStyle;
+    TSharedPtr<FSlateStyleSet> TextStyleSet;
+};
diff --git a/Engine/Plugins/Experimental/Text3D/Source/Text3D/Public/Extensions/Text3DDefaultTemplateExtension.h b/Engine/Plugins/Experimental/Text3D/Source/Text3D/Public/Extensions/Text3DDefaultTemplateExtension.h
deleted file mode 100644
diff --git a/Engine/Plugins/Experimental/Text3D/Source/Text3D/Public/Extensions/Text3DDefaultTokenExtension.h b/Engine/Plugins/Experimental/Text3D/Source/Text3D/Public/Extensions/Text3DDefaultTokenExtension.h
new file mode 100644
--- /dev/null
+++ b/Engine/Plugins/Experimental/Text3D/Source/Text3D/Public/Extensions/Text3DDefaultTokenExtension.h
@@ -0,0 +1,42 @@
+// Copyright Epic Games, Inc. All Rights Reserved.
+
+#pragma once
+
+#include "Containers/Set.h"
+#include "Text3DTokenExtensionBase.h"
+#include "Text3DDefaultTokenExtension.generated.h"
+
+class UText3DTokenBase;
+
+UCLASS(MinimalAPI)
+class UText3DDefaultTokenExtension : public UText3DTokenExtensionBase
+{
+	GENERATED_BODY()
+
+protected:
+	//~ Begin UObject
+#if WITH_EDITOR
+	virtual void PostEditUndo() override;
+	virtual void PostEditChangeProperty(FPropertyChangedEvent& InPropertyChangedEvent) override;
+#endif
+	//~ End UObject
+
+	//~ Begin UText3DTemplateExtensionBase
+	virtual const FText& GetFormattedText() const override;
+	//~ End UText3DTemplateExtensionBase
+
+	//~ Begin UText3DExtensionBase
+	virtual EText3DExtensionResult PreRendererUpdate(const UE::Text3D::Renderer::FUpdateParameters& InParameters) override;
+	virtual EText3DExtensionResult PostRendererUpdate(const UE::Text3D::Renderer::FUpdateParameters& InParameters) override;
+	//~ End UText3DExtensionBase
+
+	/**
+	 * Define your tokens and use them in text like this {MyCustomToken},
+	 * {MyCustomToken} will be swapped by the content of the token called MyCustomToken
+	 */
+	UPROPERTY(EditInstanceOnly, Instanced, Category="Token", NoClear, meta=(AllowPrivateAccess="true"))
+	TArray<TObjectPtr<UText3DTokenBase>> Tokens;
+
+	UPROPERTY(Transient)
+	FText FormattedText;
+};
diff --git a/Engine/Plugins/Experimental/Text3D/Source/Text3D/Public/Extensions/Text3DFormatStyleExtensionBase.h b/Engine/Plugins/Experimental/Text3D/Source/Text3D/Public/Extensions/Text3DFormatStyleExtensionBase.h
deleted file mode 100644
diff --git a/Engine/Plugins/Experimental/Text3D/Source/Text3D/Public/Extensions/Text3DMaterialExtensionBase.h b/Engine/Plugins/Experimental/Text3D/Source/Text3D/Public/Extensions/Text3DMaterialExtensionBase.h
--- a/Engine/Plugins/Experimental/Text3D/Source/Text3D/Public/Extensions/Text3DMaterialExtensionBase.h
+++ b/Engine/Plugins/Experimental/Text3D/Source/Text3D/Public/Extensions/Text3DMaterialExtensionBase.h
@@ -19,15 +19,59 @@ public:
 		: UText3DExtensionBase(UE::Text3D::Priority::Material)
 	{}
 
-	/** Set the material for a specific group */
+	/** Set default material for a specific group */
+	UE_DEPRECATED(5.7, "Please use SetMaterial with FMaterialParameters instead")
 	virtual void SetMaterial(EText3DGroupType InGroup, UMaterialInterface* InMaterial)
+	{
+		UE::Text3D::Material::FMaterialParameters Parameters;
+		Parameters.Group = InGroup;
+		SetMaterial(Parameters, InMaterial);
+	}
+
+	/** Get default material of a specific group */
+	UE_DEPRECATED(5.7, "Please use GetMaterial with FMaterialParameters instead")
+	virtual UMaterialInterface* GetMaterial(EText3DGroupType InGroup) const
+	{
+		UE::Text3D::Material::FMaterialParameters Parameters;
+		Parameters.Group = InGroup;
+		return GetMaterial(Parameters);
+	}
+
+	/** Set the material for a specific group */
+	virtual void SetMaterial(const UE::Text3D::Material::FMaterialParameters& InParameters, UMaterialInterface* InMaterial)
 	{
 		
 	}
 
 	/** Get the material for a specific group */
-	virtual UMaterialInterface* GetMaterial(EText3DGroupType InGroup) const
+	virtual UMaterialInterface* GetMaterial(const UE::Text3D::Material::FMaterialParameters& InParameters) const
 	{
 		return nullptr;
 	}
+
+	/** Reset all custom materials override defined */
+	virtual void ClearMaterialOverrides()
+	{
+	}
+
+	/** Register a new material override */
+	virtual void RegisterMaterialOverride(FName InTag)
+	{
+	}
+
+	/** Loop through each material available, return false to stop the operation */
+	virtual void ForEachMaterial(TFunctionRef<bool(const UE::Text3D::Material::FMaterialParameters&, UMaterialInterface*)> InFunctor) const
+	{
+	}
+
+	/** Get amount of material in use */
+	virtual int32 GetMaterialCount() const
+	{
+		return 0;
+	}
+
+	/** Retrieve material override names */
+	virtual void GetMaterialNames(TArray<FName>& OutNames) const
+	{
+	}
 };
diff --git a/Engine/Plugins/Experimental/Text3D/Source/Text3D/Public/Extensions/Text3DStyleExtensionBase.h b/Engine/Plugins/Experimental/Text3D/Source/Text3D/Public/Extensions/Text3DStyleExtensionBase.h
new file mode 100644
--- /dev/null
+++ b/Engine/Plugins/Experimental/Text3D/Source/Text3D/Public/Extensions/Text3DStyleExtensionBase.h
@@ -0,0 +1,43 @@
+// Copyright Epic Games, Inc. All Rights Reserved.
+
+#pragma once
+
+#include "Text3DExtensionBase.h"
+#include "Text3DTypes.h"
+#include "Templates/SharedPointer.h"
+#include "Text3DStyleExtensionBase.generated.h"
+
+class FSlateStyleSet;
+class UText3DStyleBase;
+struct FText3DShapedGlyphText;
+struct FTextBlockStyle;
+
+/** Extension that handles formatting styles for Text3D */
+UCLASS(MinimalAPI, Abstract)
+class UText3DStyleExtensionBase : public UText3DExtensionBase
+{
+	GENERATED_BODY()
+
+public:
+	UText3DStyleExtensionBase()
+		: UText3DExtensionBase(UE::Text3D::Priority::Style)
+	{}
+
+	/** Default style used by text when no other style is applied */
+	virtual TSharedPtr<FTextBlockStyle> GetDefaultStyle() const
+	{
+		return nullptr;
+	}
+
+	/** Custom defined styles to apply on text */
+	virtual TSharedPtr<FSlateStyleSet> GetCustomStyles() const
+	{
+		return nullptr;
+	}
+
+	/** Get the style with a specific name if any */
+	virtual UText3DStyleBase* GetStyle(FName InStyleName) const
+	{
+		return nullptr;
+	}
+};
diff --git a/Engine/Plugins/Experimental/Text3D/Source/Text3D/Public/Extensions/Text3DTemplateExtensionBase.h b/Engine/Plugins/Experimental/Text3D/Source/Text3D/Public/Extensions/Text3DTemplateExtensionBase.h
deleted file mode 100644
diff --git a/Engine/Plugins/Experimental/Text3D/Source/Text3D/Public/Extensions/Text3DTokenExtensionBase.h b/Engine/Plugins/Experimental/Text3D/Source/Text3D/Public/Extensions/Text3DTokenExtensionBase.h
new file mode 100644
--- /dev/null
+++ b/Engine/Plugins/Experimental/Text3D/Source/Text3D/Public/Extensions/Text3DTokenExtensionBase.h
@@ -0,0 +1,25 @@
+// Copyright Epic Games, Inc. All Rights Reserved.
+
+#pragma once
+
+#include "Text3DExtensionBase.h"
+#include "Text3DTypes.h"
+#include "Text3DTokenExtensionBase.generated.h"
+
+/** Extension that handles text transformation */
+UCLASS(MinimalAPI, Abstract)
+class UText3DTokenExtensionBase : public UText3DExtensionBase
+{
+	GENERATED_BODY()
+
+public:
+	UText3DTokenExtensionBase()
+		: UText3DExtensionBase(UE::Text3D::Priority::Token)
+	{}
+
+	/** Get the transformed text */
+	virtual const FText& GetFormattedText() const
+	{
+		return FText::GetEmpty();
+	}
+};
diff --git a/Engine/Plugins/Experimental/Text3D/Source/Text3D/Public/Settings/Text3DProjectSettings.h b/Engine/Plugins/Experimental/Text3D/Source/Text3D/Public/Settings/Text3DProjectSettings.h
--- a/Engine/Plugins/Experimental/Text3D/Source/Text3D/Public/Settings/Text3DProjectSettings.h
+++ b/Engine/Plugins/Experimental/Text3D/Source/Text3D/Public/Settings/Text3DProjectSettings.h
@@ -11,6 +11,7 @@ enum class EText3DMaterialBlendMode : uint8;
 class UFontFace;
 class UFont;
 class UMaterial;
+class UText3DStyleSet;
 
 /** Settings for Text3D plugin */
 UCLASS(MinimalAPI, Config=Text3D, meta=(DisplayName="Text3D"))
@@ -63,7 +64,7 @@ public:
 	UText3DProjectSettings();
 
 	TEXT3D_API UMaterial* GetDefaultMaterial() const;
-
+	TEXT3D_API UText3DStyleSet* GetDefaultStyleSet() const;
 	TEXT3D_API UFont* GetFallbackFont() const;
 	TEXT3D_API UFontFace* GetFallbackFontFace() const;
 	TEXT3D_API const FString& GetFontDirectory() const;
@@ -107,9 +108,13 @@ protected:
 	/** Default custom material used on Text */
 	UPROPERTY(Config, EditAnywhere, Category="Text3D")
 	TSoftObjectPtr<UMaterial> DefaultMaterial;
-	
-	/** Font used when the selected font is unavailable */
+
+	/** Default style rules to apply on a new Text3D instance */
 	UPROPERTY(Config, EditAnywhere, Category="Text3D")
+	TSoftObjectPtr<UText3DStyleSet> DefaultStyleSet;
+	
+	/** Default font used on new instance or when the selected font is unavailable */
+	UPROPERTY(Config, EditAnywhere, DisplayName="Default Font", Category="Text3D")
 	TSoftObjectPtr<UFont> FallbackFont;
 
 	/** Font face used as fallback when no font faces are found */
diff --git a/Engine/Plugins/Experimental/Text3D/Source/Text3D/Public/Styles/Text3DFormatStyleBase.h b/Engine/Plugins/Experimental/Text3D/Source/Text3D/Public/Styles/Text3DFormatStyleBase.h
deleted file mode 100644
diff --git a/Engine/Plugins/Experimental/Text3D/Source/Text3D/Public/Styles/Text3DStyleBase.h b/Engine/Plugins/Experimental/Text3D/Source/Text3D/Public/Styles/Text3DStyleBase.h
new file mode 100644
--- /dev/null
+++ b/Engine/Plugins/Experimental/Text3D/Source/Text3D/Public/Styles/Text3DStyleBase.h
@@ -0,0 +1,118 @@
+// Copyright Epic Games, Inc. All Rights Reserved.
+
+#pragma once
+
+#include "Containers/Array.h"
+#include "Styling/SlateTypes.h"
+#include "Text3DTypes.h"
+#include "UObject/Object.h"
+#include "Text3DStyleBase.generated.h"
+
+class UFont;
+struct FTextBlockStyle;
+struct FTypefaceEntry;
+
+/** Text3D rich text format styles */
+UCLASS(MinimalAPI, EditInlineNew, ClassGroup=Text3D, DisplayName="FormatStyle", AutoExpandCategories=(Style))
+class UText3DStyleBase : public UObject
+{
+	GENERATED_BODY()
+
+public:
+	UText3DStyleBase();
+
+	UFUNCTION(BlueprintCallable, Category = "Text3D|Style")
+	TEXT3D_API void SetStyleName(FName InName);
+
+	UFUNCTION(BlueprintPure, Category = "Text3D|Style")
+	FName GetStyleName() const
+	{
+		return StyleName;
+	}
+
+	UFUNCTION(BlueprintCallable, Category = "Text3D|Style")
+	TEXT3D_API void SetOverrideFont(bool bInOverride);
+
+	UFUNCTION(BlueprintPure, Category = "Text3D|Style")
+	bool GetOverrideFont() const
+	{
+		return bOverrideFont;
+	}
+
+	UFUNCTION(BlueprintCallable, Category = "Text3D|Style")
+	TEXT3D_API void SetFont(UFont* InFont);
+
+	UFUNCTION(BlueprintPure, Category = "Text3D|Style")
+	UFont* GetFont() const
+	{
+		return Font;
+	}
+
+	UFUNCTION(BlueprintCallable, Category = "Text3D|Style")
+	TEXT3D_API void SetFontTypeface(FName InTypeface);
+
+	UFUNCTION(BlueprintPure, Category = "Text3D|Style")
+	FName GetFontTypeface() const
+	{
+		return FontTypeface;
+	}
+
+	UFUNCTION(BlueprintCallable, Category = "Text3D|Style")
+	TEXT3D_API void SetOverrideFrontColor(bool bInOverride);
+
+	UFUNCTION(BlueprintPure, Category = "Text3D|Style")
+	bool GetOverrideFrontColor() const
+	{
+		return bOverrideFrontColor;
+	}
+
+	UFUNCTION(BlueprintCallable, Category = "Text3D|Style")
+	TEXT3D_API void SetFrontColor(FLinearColor InColor);
+
+	UFUNCTION(BlueprintPure, Category = "Text3D|Style")
+	FLinearColor GetFrontColor() const
+	{
+		return FrontColor;
+	}
+
+	/** Sets the proper values on the text block style */
+	void ConfigureStyle(FTextBlockStyle& OutStyle) const;
+
+protected:
+	//~ Begin UObject
+#if WITH_EDITOR
+	virtual void PostEditChangeProperty(FPropertyChangedEvent& InPropertyChangedEvent) override;
+#endif
+	//~ End UObject
+
+	void OnStylePropertiesChanged();
+	void OnFontChanged();
+
+	UFUNCTION()
+	TArray<FName> GetTypefaceNames() const;
+	TArrayView<const FTypefaceEntry> GetAvailableTypefaces() const;
+
+	/** Style name used in text */
+	UPROPERTY(EditAnywhere, Setter, Getter, Category="Style", meta=(AllowPrivateAccess="true"))
+	FName StyleName = NAME_None;
+
+	/** Behavior for the font rule of this style */
+	UPROPERTY(EditAnywhere, Setter="SetOverrideFont", Getter="GetOverrideFont", Category="Style", meta=(AllowPrivateAccess="true", InlineEditConditionToggle))
+	bool bOverrideFont = true;
+
+	/** Text font defines the style of rendered characters */
+	UPROPERTY(EditAnywhere, Setter, Getter, Category="Style", meta=(EditCondition="bOverrideFont", AdvancedFontPicker, AllowPrivateAccess="true"))
+	TObjectPtr<UFont> Font;
+
+	/** Text font face, subset within font like bold, italic, regular */
+	UPROPERTY(EditAnywhere, Setter, Getter, Category="Style", meta=(EditCondition="bOverrideFont", GetOptions="GetTypefaceNames", AllowPrivateAccess="true"))
+	FName FontTypeface;
+
+	/** Behavior for the front color rule of this style */
+	UPROPERTY(EditAnywhere, Setter="SetOverrideFrontColor", Getter="GetOverrideFrontColor", Category="Style", meta=(AllowPrivateAccess="true", InlineEditConditionToggle))
+	bool bOverrideFrontColor = true;
+
+	/** Text front color */
+	UPROPERTY(EditAnywhere, Category="Style", meta=(EditCondition="bOverrideFrontColor", AllowPrivateAccess="true"))
+	FLinearColor FrontColor = FLinearColor::White;
+};
diff --git a/Engine/Plugins/Experimental/Text3D/Source/Text3D/Public/Styles/Text3DStyleSet.h b/Engine/Plugins/Experimental/Text3D/Source/Text3D/Public/Styles/Text3DStyleSet.h
new file mode 100644
--- /dev/null
+++ b/Engine/Plugins/Experimental/Text3D/Source/Text3D/Public/Styles/Text3DStyleSet.h
@@ -0,0 +1,42 @@
+// Copyright Epic Games, Inc. All Rights Reserved.
+
+#pragma once
+
+#include "UObject/Object.h"
+#include "Text3DStyleSet.generated.h"
+
+class UText3DStyleBase;
+
+/** Text3D Style Set used to store multiple reusable styles */
+UCLASS(MinimalAPI, DisplayName="Text3D Style Set", AutoExpandCategories=(Style))
+class UText3DStyleSet : public UObject
+{
+	GENERATED_BODY()
+
+public:
+	DECLARE_MULTICAST_DELEGATE_OneParam(FOnStyleSetUpdated, UText3DStyleSet* /** InStyleSet */)
+
+	static FOnStyleSetUpdated::RegistrationType& OnStyleSetUpdated()
+	{
+		return OnStyleSetUpdatedDelegate;
+	}
+
+	TConstArrayView<TObjectPtr<UText3DStyleBase>> GetStyles() const
+	{
+		return Styles;
+	}
+
+protected:
+	static FOnStyleSetUpdated OnStyleSetUpdatedDelegate;
+
+	//~ Begin UObject
+#if WITH_EDITOR
+	virtual void PostEditChangeProperty(FPropertyChangedEvent& InPropertyChangedEvent) override;
+#endif
+	//~ End UObject
+
+	void OnStyleSetPropertiesChanged();
+
+	UPROPERTY(EditAnywhere, Instanced, Category="Style", NoClear, meta=(AllowPrivateAccess="true"))
+	TArray<TObjectPtr<UText3DStyleBase>> Styles;
+};
\ No newline at end of file
diff --git a/Engine/Plugins/Experimental/Text3D/Source/Text3D/Public/Templates/Text3DTemplateBase.h b/Engine/Plugins/Experimental/Text3D/Source/Text3D/Public/Templates/Text3DTemplateBase.h
deleted file mode 100644
diff --git a/Engine/Plugins/Experimental/Text3D/Source/Text3D/Public/Text3DComponent.h b/Engine/Plugins/Experimental/Text3D/Source/Text3D/Public/Text3DComponent.h
--- a/Engine/Plugins/Experimental/Text3D/Source/Text3D/Public/Text3DComponent.h
+++ b/Engine/Plugins/Experimental/Text3D/Source/Text3D/Public/Text3DComponent.h
@@ -13,12 +13,12 @@ class UFont;
 class UMaterialInterface;
 class UText3DCharacterBase;
 class UText3DCharacterExtensionBase;
-class UText3DFormatStyleExtensionBase;
 class UText3DGeometryExtensionBase;
 class UText3DLayoutEffectBase;
 class UText3DLayoutExtensionBase;
 class UText3DMaterialExtensionBase;
-class UText3DTemplateExtensionBase;
+class UText3DStyleExtensionBase;
+class UText3DTokenExtensionBase;
 class UText3DRendererBase;
 class UText3DRenderingExtensionBase;
 struct FTypefaceEntry;
@@ -31,7 +31,7 @@ UCLASS(MinimalAPI
 class UText3DComponent : public USceneComponent
 {
 	// Needed to support older API (see FormatText)
-	friend class UText3DDefaultTemplateExtension;
+	friend class UText3DDefaultTokenExtension;
 
 	GENERATED_BODY()
 
@@ -273,31 +273,31 @@ public:
 	TEXT3D_API UText3DRendererBase* GetTextRenderer() const;
 
 	UFUNCTION(BlueprintCallable, Category="Text3D", meta=(DeterminesOutputType="InExtensionClass"))
-	UText3DTemplateExtensionBase* GetTemplateExtension(TSubclassOf<UText3DTemplateExtensionBase> InExtensionClass);
+	UText3DTokenExtensionBase* GetTokenExtension(TSubclassOf<UText3DTokenExtensionBase> InExtensionClass);
 	
-	UText3DTemplateExtensionBase* GetTemplateExtension() const
+	UText3DTokenExtensionBase* GetTokenExtension() const
 	{
-		return TemplateExtension;
+		return TokenExtension;
 	}
 	
-	template <typename T UE_REQUIRES(TIsDerivedFrom<T, UText3DTemplateExtensionBase>::Value)>
-	T* GetCastedTemplateExtension() const
+	template <typename T UE_REQUIRES(TIsDerivedFrom<T, UText3DTokenExtensionBase>::Value)>
+	T* GetCastedTokenExtension() const
 	{
-		return Cast<T>(TemplateExtension);
+		return Cast<T>(TokenExtension);
 	}
 
 	UFUNCTION(BlueprintCallable, Category="Text3D", meta=(DeterminesOutputType="InExtensionClass"))
-	UText3DFormatStyleExtensionBase* GetFormatStyleExtension(TSubclassOf<UText3DFormatStyleExtensionBase> InExtensionClass);
+	UText3DStyleExtensionBase* GetStyleExtension(TSubclassOf<UText3DStyleExtensionBase> InExtensionClass);
 	
-	UText3DFormatStyleExtensionBase* GetFormatStyleExtension() const
+	UText3DStyleExtensionBase* GetStyleExtension() const
 	{
-		return FormatStyleExtension;
+		return StyleExtension;
 	}
 	
-	template <typename T UE_REQUIRES(TIsDerivedFrom<T, UText3DFormatStyleExtensionBase>::Value)>
-	T* GetCastedFormatStyleExtension() const
+	template <typename T UE_REQUIRES(TIsDerivedFrom<T, UText3DStyleExtensionBase>::Value)>
+	T* GetCastedStyleExtension() const
 	{
-		return Cast<T>(FormatStyleExtension);
+		return Cast<T>(StyleExtension);
 	}
 
 	UFUNCTION(BlueprintCallable, Category="Text3D", meta=(DeterminesOutputType="InExtensionClass"))
@@ -431,8 +431,13 @@ protected:
 	/** Legacy : used to resolve font by their names */
 	TEXT3D_API void SetFontByName(const FString& InFontName);
 
+	/** Refresh text to force update render state */
 	UFUNCTION(CallInEditor, Category="Utilities")
 	void ForceUpdateText();
+
+	/** Open Text project settings editor window */
+	UFUNCTION(CallInEditor, Category="Utilities")
+	void OpenProjectSettings();
 #endif
 
 	/** Get the type faces supported for the current font */
@@ -475,11 +480,11 @@ protected:
 	UPROPERTY(EditAnywhere, Getter = "GetEnforceUpperCase", Setter = "SetEnforceUpperCase", Category = "Text", meta = (AllowPrivateAccess="true"))
 	bool bEnforceUpperCase = false;
 
-	UPROPERTY(VisibleInstanceOnly, Instanced, Category="Template", meta=(AllowPrivateAccess="true"))
-	TObjectPtr<UText3DTemplateExtensionBase> TemplateExtension;
+	UPROPERTY(VisibleInstanceOnly, Instanced, Category="Token", meta=(AllowPrivateAccess="true"))
+	TObjectPtr<UText3DTokenExtensionBase> TokenExtension;
 
 	UPROPERTY(VisibleInstanceOnly, Instanced, Category="Style", meta=(AllowPrivateAccess="true"))
-	TObjectPtr<UText3DFormatStyleExtensionBase> FormatStyleExtension;
+	TObjectPtr<UText3DStyleExtensionBase> StyleExtension;
 
 	UPROPERTY(VisibleInstanceOnly, Instanced, Category="Character", meta=(AllowPrivateAccess="true"))
 	TObjectPtr<UText3DCharacterExtensionBase> CharacterExtension;
diff --git a/Engine/Plugins/Experimental/Text3D/Source/Text3D/Public/Text3DTypes.h b/Engine/Plugins/Experimental/Text3D/Source/Text3D/Public/Text3DTypes.h
--- a/Engine/Plugins/Experimental/Text3D/Source/Text3D/Public/Text3DTypes.h
+++ b/Engine/Plugins/Experimental/Text3D/Source/Text3D/Public/Text3DTypes.h
@@ -2,10 +2,13 @@
 
 #pragma once
 
+#include "Containers/Array.h"
 #include "Misc/EnumClassFlags.h"
+#include "UObject/ObjectPtr.h"
 #include "Text3DTypes.generated.h"
 
 class UDynamicMesh;
+class UMaterialInterface;
 class UStaticMesh;
 class UWorld;
 struct FText3DFontFaceCache;
@@ -73,6 +76,7 @@ enum class EText3DOutlineType : uint8
 	Fill
 };
 
+UENUM()
 enum class EText3DGroupType : uint8
 {
 	Front = 0,
@@ -193,17 +197,17 @@ struct FText3DMaterialGroupKey
 	GENERATED_BODY()
 
 	FText3DMaterialGroupKey() = default;
-	explicit FText3DMaterialGroupKey(FText3DMaterialKey InKey, EText3DGroupType InGroup, EText3DMaterialStyle InStyle)
+	explicit FText3DMaterialGroupKey(FText3DMaterialKey InKey, EText3DGroupType InGroup, FName InTag)
 		: Key(InKey)
 		, Group(InGroup)
-		, Style(InStyle)
+		, Tag(InTag)
 	{}
 
 	bool operator==(const FText3DMaterialGroupKey& Other) const
 	{
 		return Key == Other.Key
 			&& Group == Other.Group
-			&& Style == Other.Style;
+			&& Tag == Other.Tag;
 	}
 
 	bool operator!=(const FText3DMaterialGroupKey& Other) const
@@ -213,16 +217,22 @@ struct FText3DMaterialGroupKey
 
 	friend uint32 GetTypeHash(const FText3DMaterialGroupKey& InMaterialSettings)
 	{
-		return HashCombineFast(HashCombineFast(GetTypeHash(InMaterialSettings.Key), GetTypeHash(InMaterialSettings.Group)), GetTypeHash(InMaterialSettings.Style));
+		return HashCombineFast(
+			HashCombineFast(GetTypeHash(InMaterialSettings.Key), GetTypeHash(InMaterialSettings.Group)),
+			GetTypeHash(InMaterialSettings.Tag)
+		);
 	}
 
+	/** Defines what the material specs are */
 	UPROPERTY()
 	FText3DMaterialKey Key;
 
+	/** Defines to which group this material is applied */
 	EText3DGroupType Group = EText3DGroupType::Front;
 
+	/** Defines a tag to differentiate materials by style */
 	UPROPERTY()
-	EText3DMaterialStyle Style = EText3DMaterialStyle::Invalid;
+	FName Tag = NAME_None;
 };
 
 USTRUCT()
@@ -290,11 +300,39 @@ struct FText3DCachedMesh
 	uint32 RefCount = 0;
 };
 
+USTRUCT()
+struct FText3DMaterialOverride
+{
+	GENERATED_BODY()
+
+	UPROPERTY()
+	FName Tag = NAME_None;
+
+	/** Materials where index = EText3DGroupType */
+	UPROPERTY()
+	TArray<TObjectPtr<UMaterialInterface>> Materials;
+
+	bool operator==(const FText3DMaterialOverride& Other) const
+	{
+		return Tag == Other.Tag;
+	}
+
+	bool operator!=(const FText3DMaterialOverride& Other) const
+	{
+		return !(*this == Other);
+	}
+
+	friend uint32 GetTypeHash(const FText3DMaterialOverride& InElement)
+	{
+		return GetTypeHash(InElement.Tag);
+	}
+};
+
 namespace UE::Text3D
 {
 	namespace Priority
 	{
-		static constexpr uint16 Template = 0;
+		static constexpr uint16 Token = 0;
 		static constexpr uint16 Geometry = 1;
 		static constexpr uint16 Style = 2;
 		static constexpr uint16 Layout = 3;
@@ -362,4 +400,16 @@ namespace UE::Text3D
 			uint32 GlyphMeshHash = 0;
 		};
 	}
+
+	namespace Material
+	{
+		struct FMaterialParameters
+		{
+			/** Tag used to identify a material section, none = default */
+			FName Tag = NAME_None;
+
+			/** Group used to identify an entry in the material section */
+			EText3DGroupType Group = EText3DGroupType::Front;
+		};
+	}
 }
\ No newline at end of file
diff --git a/Engine/Plugins/Experimental/Text3D/Source/Text3D/Public/Tokens/Text3DTokenBase.h b/Engine/Plugins/Experimental/Text3D/Source/Text3D/Public/Tokens/Text3DTokenBase.h
new file mode 100644
--- /dev/null
+++ b/Engine/Plugins/Experimental/Text3D/Source/Text3D/Public/Tokens/Text3DTokenBase.h
@@ -0,0 +1,46 @@
+// Copyright Epic Games, Inc. All Rights Reserved.
+
+#pragma once
+
+#include "UObject/Object.h"
+#include "Text3DTokenBase.generated.h"
+
+/** Text3D token to swap variable by its content */
+UCLASS(MinimalAPI, EditInlineNew, ClassGroup=Text3D, DisplayName="Token", AutoExpandCategories=(Token))
+class UText3DTokenBase : public UObject
+{
+	GENERATED_BODY()
+
+public:
+	void SetTokenName(FName InName);
+	FName GetTokenName() const
+	{
+		return TokenName;
+	}
+
+	void SetContent(const FText& InContent);
+	const FText& GetContent() const
+	{
+		return Content;
+	}
+
+	/** Collects named tokens and their content for replacement */
+	virtual void CollectTokens(FFormatNamedArguments& InNamedArguments);
+
+protected:
+	//~ Begin UObject
+#if WITH_EDITOR
+	virtual void PostEditChangeProperty(FPropertyChangedEvent& InPropertyChangedEvent) override;
+#endif
+	//~ End UObject
+
+	void OnTokenPropertiesChanged();
+
+	/** Token variable that will be swapped */
+	UPROPERTY(EditAnywhere, Setter, Getter, Category="Token", meta=(AllowPrivateAccess="true"))
+	FName TokenName = NAME_None;
+
+	/** Content that replaces token in the text */
+	UPROPERTY(EditAnywhere, Setter, Getter, Category="Token", meta=(AllowPrivateAccess="true"))
+	FText Content;
+};
diff --git a/Engine/Plugins/Experimental/Text3D/Source/Text3DEditor/Private/AssetTypeActions/AssetDefinition_Text3DStyleSet.cpp b/Engine/Plugins/Experimental/Text3D/Source/Text3DEditor/Private/AssetTypeActions/AssetDefinition_Text3DStyleSet.cpp
new file mode 100644
--- /dev/null
+++ b/Engine/Plugins/Experimental/Text3D/Source/Text3DEditor/Private/AssetTypeActions/AssetDefinition_Text3DStyleSet.cpp
@@ -0,0 +1,20 @@
+// Copyright Epic Games, Inc. All Rights Reserved.
+
+#include "AssetTypeActions/AssetDefinition_Text3DStyleSet.h"
+
+#include "Styles/Text3DStyleSet.h"
+
+FText UAssetDefinition_Text3DStyleSet::GetAssetDisplayName() const
+{
+	return NSLOCTEXT("AssetTypeActions", "AssetTypeActions_Text3DStyleSet", "Text3D Style Set");
+}
+
+FLinearColor UAssetDefinition_Text3DStyleSet::GetAssetColor() const
+{
+	return FLinearColor(FColor(128, 255, 128));
+}
+
+TSoftClassPtr<UObject> UAssetDefinition_Text3DStyleSet::GetAssetClass() const
+{
+	return UText3DStyleSet::StaticClass();
+}
\ No newline at end of file
diff --git a/Engine/Plugins/Experimental/Text3D/Source/Text3DEditor/Private/AssetTypeActions/AssetDefinition_Text3DStyleSet.h b/Engine/Plugins/Experimental/Text3D/Source/Text3DEditor/Private/AssetTypeActions/AssetDefinition_Text3DStyleSet.h
new file mode 100644
--- /dev/null
+++ b/Engine/Plugins/Experimental/Text3D/Source/Text3DEditor/Private/AssetTypeActions/AssetDefinition_Text3DStyleSet.h
@@ -0,0 +1,20 @@
+// Copyright Epic Games, Inc. All Rights Reserved.
+
+#pragma once
+
+#include "AssetDefinitionDefault.h"
+#include "AssetDefinition_Text3DStyleSet.generated.h"
+
+/* Defines a new asset : StyleSet */
+UCLASS()
+class UAssetDefinition_Text3DStyleSet : public UAssetDefinitionDefault
+{
+	GENERATED_BODY()
+
+public:
+	//~ Begin UAssetDefinition
+	virtual FText GetAssetDisplayName() const override;
+	virtual FLinearColor GetAssetColor() const override;
+	virtual TSoftClassPtr<UObject> GetAssetClass() const override;
+	//~ End UAssetDefinition
+};
diff --git a/Engine/Plugins/Experimental/Text3D/Source/Text3DEditor/Private/Customizations/Text3DEditorTextComponentDetailCustomization.cpp b/Engine/Plugins/Experimental/Text3D/Source/Text3DEditor/Private/Customizations/Text3DEditorTextComponentDetailCustomization.cpp
--- a/Engine/Plugins/Experimental/Text3D/Source/Text3DEditor/Private/Customizations/Text3DEditorTextComponentDetailCustomization.cpp
+++ b/Engine/Plugins/Experimental/Text3D/Source/Text3DEditor/Private/Customizations/Text3DEditorTextComponentDetailCustomization.cpp
@@ -22,7 +22,7 @@ namespace UE::Text3DEditor::Customization
 		const TArray<FName> CategoryNames =
 		{
 			TEXT("Text"),
-			TEXT("Template"),
+			TEXT("Token"),
 			TEXT("Style"),
 			TEXT("Character"),
 			TEXT("Layout"),
@@ -62,7 +62,7 @@ namespace UE::Text3DEditor::Customization
 			const FName TextSectionName = TEXT("Text");
 			const TSharedRef<FPropertySection> TextSection = PropertyModule.FindOrCreateSection(ComponentClassName, TextSectionName, FText::FromName(TextSectionName));
 			TextSection->AddCategory(TEXT("Text"));
-			TextSection->AddCategory(TEXT("Template"));
+			TextSection->AddCategory(TEXT("Token"));
 
 			const FName StyleSectionName = TEXT("Style");
 			const TSharedRef<FPropertySection> MaterialSection = PropertyModule.FindOrCreateSection(ComponentClassName, StyleSectionName, FText::FromName(StyleSectionName));
diff --git a/Engine/Plugins/Experimental/Text3D/Source/Text3DEditor/Private/Factories/Text3DEditorStyleSetFactory.cpp b/Engine/Plugins/Experimental/Text3D/Source/Text3DEditor/Private/Factories/Text3DEditorStyleSetFactory.cpp
new file mode 100644
--- /dev/null
+++ b/Engine/Plugins/Experimental/Text3D/Source/Text3DEditor/Private/Factories/Text3DEditorStyleSetFactory.cpp
@@ -0,0 +1,23 @@
+// Copyright Epic Games, Inc. All Rights Reserved.
+
+#include "Factories/Text3DEditorStyleSetFactory.h"
+
+#include "Styles/Text3DStyleSet.h"
+
+UText3DEditorStyleSetFactory::UText3DEditorStyleSetFactory()
+{
+	SupportedClass = UText3DStyleSet::StaticClass();
+	bCreateNew = true;
+	bEditorImport = false;
+	bEditAfterNew = true;
+}
+
+UObject* UText3DEditorStyleSetFactory::FactoryCreateNew(UClass* InClass, UObject* InParent, FName InName, EObjectFlags InFlags, UObject* InContext, FFeedbackContext* InWarn)
+{
+	if (!ensure(SupportedClass == InClass))
+	{
+		return nullptr;
+	}
+
+	return NewObject<UText3DStyleSet>(InParent, InName, InFlags);
+}
diff --git a/Engine/Plugins/Experimental/Text3D/Source/Text3DEditor/Private/Factories/Text3DEditorStyleSetFactory.h b/Engine/Plugins/Experimental/Text3D/Source/Text3DEditor/Private/Factories/Text3DEditorStyleSetFactory.h
new file mode 100644
--- /dev/null
+++ b/Engine/Plugins/Experimental/Text3D/Source/Text3DEditor/Private/Factories/Text3DEditorStyleSetFactory.h
@@ -0,0 +1,19 @@
+// Copyright Epic Games, Inc. All Rights Reserved.
+
+#pragma once
+
+#include "Factories/Factory.h"
+#include "Text3DEditorStyleSetFactory.generated.h"
+
+/** Factory that creates a Style Set asset */
+UCLASS()
+class UText3DEditorStyleSetFactory : public UFactory
+{
+	GENERATED_BODY()
+
+	UText3DEditorStyleSetFactory();
+
+	//~ Begin UFactory
+	virtual UObject* FactoryCreateNew(UClass* InClass, UObject* InParent, FName InName, EObjectFlags InFlags, UObject* InContext, FFeedbackContext* InWarn) override;
+	//~ End UFactory
+};
diff --git a/Engine/Plugins/Experimental/Text3D/Source/Text3DEditor/Text3DEditor.Build.cs b/Engine/Plugins/Experimental/Text3D/Source/Text3DEditor/Text3DEditor.Build.cs
--- a/Engine/Plugins/Experimental/Text3D/Source/Text3DEditor/Text3DEditor.Build.cs
+++ b/Engine/Plugins/Experimental/Text3D/Source/Text3DEditor/Text3DEditor.Build.cs
@@ -8,6 +8,7 @@ public class Text3DEditor : ModuleRules
 	public Text3DEditor(ReadOnlyTargetRules Target) : base(Target)
 	{
 		PrivateDependencyModuleNames.AddRange(new string[] {
+			"AssetDefinition",
 			"Core",
 			"CoreUObject",
 			"DeveloperSettings",
diff --git a/Engine/Plugins/VirtualProduction/Avalanche/Source/AvalancheMask/Private/Handling/AvaMaskText3DActorMaterialCollectionHandle.cpp b/Engine/Plugins/VirtualProduction/Avalanche/Source/AvalancheMask/Private/Handling/AvaMaskText3DActorMaterialCollectionHandle.cpp
--- a/Engine/Plugins/VirtualProduction/Avalanche/Source/AvalancheMask/Private/Handling/AvaMaskText3DActorMaterialCollectionHandle.cpp
+++ b/Engine/Plugins/VirtualProduction/Avalanche/Source/AvalancheMask/Private/Handling/AvaMaskText3DActorMaterialCollectionHandle.cpp
@@ -88,10 +88,18 @@ TArray<TSharedPtr<IAvaMaskMaterialHandle>> FAvaMaskText3DActorMaterialCollection
 
 			const UText3DMaterialExtensionBase* MaterialExtension = TextComponent->GetMaterialExtension();
 
-			UMaterialInterface* FrontMaterial = MaterialExtension->GetMaterial(EText3DGroupType::Front);
-			UMaterialInterface* BackMaterial = MaterialExtension->GetMaterial(EText3DGroupType::Back);
-			UMaterialInterface* ExtrudeMaterial = MaterialExtension->GetMaterial(EText3DGroupType::Extrude);
-			UMaterialInterface* BevelMaterial = MaterialExtension->GetMaterial(EText3DGroupType::Bevel);
+			UE::Text3D::Material::FMaterialParameters Parameters;
+			Parameters.Group = EText3DGroupType::Front;
+			UMaterialInterface* FrontMaterial = MaterialExtension->GetMaterial(Parameters);
+
+			Parameters.Group = EText3DGroupType::Back;
+			UMaterialInterface* BackMaterial = MaterialExtension->GetMaterial(Parameters);
+
+			Parameters.Group = EText3DGroupType::Extrude;
+			UMaterialInterface* ExtrudeMaterial = MaterialExtension->GetMaterial(Parameters);
+
+			Parameters.Group = EText3DGroupType::Bevel;
+			UMaterialInterface* BevelMaterial = MaterialExtension->GetMaterial(Parameters);
 
 			MaterialHandles.Init(nullptr, static_cast<int32>(EText3DGroupType::TypeCount));
 			MaterialHandles[static_cast<int32>(EText3DGroupType::Front)] = FrontMaterial != nullptr ? HandleSubsystem->MakeHandle<IAvaMaskMaterialHandle>(FrontMaterial, UE::AvaMask::Internal::HandleTag) : nullptr;
