MotionDesign : Text3D
- Added font size for default text and per custom style, text font size matches slate font size
- Fixed extension property undo not updating state
- Fixed duplicated text would result in character leftovers

#DisabledValidations NonUnityAndStaticAnalysis

diff --git a/Engine/Plugins/Experimental/Text3D/Source/Text3D/Private/Extensions/Text3DDefaultGeometryExtension.cpp b/Engine/Plugins/Experimental/Text3D/Source/Text3D/Private/Extensions/Text3DDefaultGeometryExtension.cpp
--- a/Engine/Plugins/Experimental/Text3D/Source/Text3D/Private/Extensions/Text3DDefaultGeometryExtension.cpp
+++ b/Engine/Plugins/Experimental/Text3D/Source/Text3D/Private/Extensions/Text3DDefaultGeometryExtension.cpp
@@ -125,6 +125,12 @@ void UText3DDefaultGeometryExtension::SetBevelSegments(const int32 Value)
 }
 
 #if WITH_EDITOR
+void UText3DDefaultGeometryExtension::PostEditUndo()
+{
+	Super::PostEditUndo();
+	OnGeometryOptionsChanged();
+}
+
 void UText3DDefaultGeometryExtension::PostEditChangeProperty(FPropertyChangedEvent& InEvent)
 {
 	Super::PostEditChangeProperty(InEvent);
@@ -178,10 +184,10 @@ EText3DExtensionResult UText3DDefaultGeometryExtension::PreRendererUpdate(const
 
 	const UText3DComponent* TextComponent = GetText3DComponent();
 	const int32 TypefaceIndex = TextComponent->GetTypeFaceIndex();
-	const UFont* Font = TextComponent->GetFont();
 
 	if (TypefaceIndex == INDEX_NONE)
 	{
+		const UFont* Font = TextComponent->GetFont();
 		UE_LOG(LogText3D, Error, TEXT("Invalid typeface index returned for font %s : %i"), Font ? *Font->GetName() : TEXT("invalid"), TypefaceIndex);
 		return EText3DExtensionResult::Failed;
 	}
diff --git a/Engine/Plugins/Experimental/Text3D/Source/Text3D/Private/Extensions/Text3DDefaultLayoutExtension.cpp b/Engine/Plugins/Experimental/Text3D/Source/Text3D/Private/Extensions/Text3DDefaultLayoutExtension.cpp
--- a/Engine/Plugins/Experimental/Text3D/Source/Text3D/Private/Extensions/Text3DDefaultLayoutExtension.cpp
+++ b/Engine/Plugins/Experimental/Text3D/Source/Text3D/Private/Extensions/Text3DDefaultLayoutExtension.cpp
@@ -320,13 +320,8 @@ EText3DExtensionResult UText3DDefaultLayoutExtension::PreRendererUpdate(const UE
 					}
 				}
 
-				constexpr bool bReset = true;
-				FTransform& CharacterTransform = Character->GetTransform(bReset);
-
-				FBox GlyphBounds;
-				FVector GlyphOffset;
-				Character->GetGlyphMeshBoundsAndOffset(GlyphBounds, GlyphOffset);
-
+				const FText3DCachedMesh* CachedMesh = Character->GetGlyphMesh();
+				const FBox GlyphBounds = CachedMesh->MeshBounds;
 				const FVector GlyphSize = GlyphBounds.GetSize();
 
 				FVector Shift = FVector::ZeroVector;
@@ -348,7 +343,22 @@ EText3DExtensionResult UText3DDefaultLayoutExtension::PreRendererUpdate(const UE
 					Shift.Y += Tracking + GlyphText->Kernings[CharacterIndex];
 				}
 
-				CharacterTransform.Accumulate(FTransform(Location + Shift));
+				FTransform Transform(Location + Shift);
+				const FVector2D FaceSize = CachedMesh->FontFaceGlyphSize;
+
+				if (FaceSize.IsNearlyZero())
+				{
+					UE_LOG(LogText3D, Warning, TEXT("Font face glyph size is nearly zero, cannot scale mesh properly to match font size"))
+				}
+				else
+				{
+					const FVector Ratio(1.f, GlyphEntry->Size / FaceSize.X, GlyphEntry->Size / FaceSize.Y);
+					Transform.SetScale3D(Ratio);
+				}
+
+				constexpr bool bReset = true;
+				FTransform& CharacterTransform = Character->GetTransform(bReset);
+				CharacterTransform.Accumulate(Transform);
 				CharacterIndex++;
 			}
 
diff --git a/Engine/Plugins/Experimental/Text3D/Source/Text3D/Private/Extensions/Text3DDefaultRenderingExtension.cpp b/Engine/Plugins/Experimental/Text3D/Source/Text3D/Private/Extensions/Text3DDefaultRenderingExtension.cpp
--- a/Engine/Plugins/Experimental/Text3D/Source/Text3D/Private/Extensions/Text3DDefaultRenderingExtension.cpp
+++ b/Engine/Plugins/Experimental/Text3D/Source/Text3D/Private/Extensions/Text3DDefaultRenderingExtension.cpp
@@ -58,6 +58,12 @@ void UText3DDefaultRenderingExtension::SetHoldout(bool bInHoldout)
 }
 
 #if WITH_EDITOR
+void UText3DDefaultRenderingExtension::PostEditUndo()
+{
+	Super::PostEditUndo();
+	OnRenderingOptionsChanged();
+}
+
 void UText3DDefaultRenderingExtension::PostEditChangeProperty(FPropertyChangedEvent& InEvent)
 {
 	Super::PostEditChangeProperty(InEvent);
diff --git a/Engine/Plugins/Experimental/Text3D/Source/Text3D/Private/Extensions/Text3DDefaultStyleExtension.cpp b/Engine/Plugins/Experimental/Text3D/Source/Text3D/Private/Extensions/Text3DDefaultStyleExtension.cpp
--- a/Engine/Plugins/Experimental/Text3D/Source/Text3D/Private/Extensions/Text3DDefaultStyleExtension.cpp
+++ b/Engine/Plugins/Experimental/Text3D/Source/Text3D/Private/Extensions/Text3DDefaultStyleExtension.cpp
@@ -112,7 +112,7 @@ EText3DExtensionResult UText3DDefaultStyleExtension::PreRendererUpdate(const UE:
 			DefaultFontStyle = MakeShared<FTextBlockStyle>();
 		}
 
-		FSlateFontInfo FontInfo(Font, UE::Text3D::Metrics::SlateFontSize);
+		FSlateFontInfo FontInfo(Font, Text3DComponent->GetFontSize());
 		FontInfo.CompositeFont = Font && Font->GetCompositeFont() ? MakeShared<FCompositeFont>(*Font->GetCompositeFont()) : FStyleDefaults::GetFontInfo().CompositeFont;
 		FontInfo.TypefaceFontName = Text3DComponent->GetTypeface();
 		FontInfo.FontFallback = EFontFallback::FF_NoFallback;
diff --git a/Engine/Plugins/Experimental/Text3D/Source/Text3D/Private/Extensions/Text3DLayoutTransformEffect.cpp b/Engine/Plugins/Experimental/Text3D/Source/Text3D/Private/Extensions/Text3DLayoutTransformEffect.cpp
--- a/Engine/Plugins/Experimental/Text3D/Source/Text3D/Private/Extensions/Text3DLayoutTransformEffect.cpp
+++ b/Engine/Plugins/Experimental/Text3D/Source/Text3D/Private/Extensions/Text3DLayoutTransformEffect.cpp
@@ -183,6 +183,12 @@ UText3DLayoutTransformEffect::UText3DLayoutTransformEffect()
 }
 
 #if WITH_EDITOR
+void UText3DLayoutTransformEffect::PostEditUndo()
+{
+	Super::PostEditUndo();
+	OnTransformOptionsChanged();
+}
+
 void UText3DLayoutTransformEffect::PostEditChangeProperty(FPropertyChangedEvent& InEvent)
 {
 	Super::PostEditChangeProperty(InEvent);
diff --git a/Engine/Plugins/Experimental/Text3D/Source/Text3D/Private/LayoutBuilders/Text3DLayoutShaper.cpp b/Engine/Plugins/Experimental/Text3D/Source/Text3D/Private/LayoutBuilders/Text3DLayoutShaper.cpp
--- a/Engine/Plugins/Experimental/Text3D/Source/Text3D/Private/LayoutBuilders/Text3DLayoutShaper.cpp
+++ b/Engine/Plugins/Experimental/Text3D/Source/Text3D/Private/LayoutBuilders/Text3DLayoutShaper.cpp
@@ -79,6 +79,7 @@ void FText3DLayoutShaper::ShapeBidirectionalText(
 			{
 				FGlyphEntry& GlyphEntry = ShapedLine.Glyphs.Add_GetRef(ShapedGlyphEntry);
 				GlyphEntry.StyleTag = StyleTag;
+				GlyphEntry.Size = FontInfo.Size;
 			}
 		}
 
diff --git a/Engine/Plugins/Experimental/Text3D/Source/Text3D/Private/LayoutBuilders/Text3DShapedGlyphLine.h b/Engine/Plugins/Experimental/Text3D/Source/Text3D/Private/LayoutBuilders/Text3DShapedGlyphLine.h
--- a/Engine/Plugins/Experimental/Text3D/Source/Text3D/Private/LayoutBuilders/Text3DShapedGlyphLine.h
+++ b/Engine/Plugins/Experimental/Text3D/Source/Text3D/Private/LayoutBuilders/Text3DShapedGlyphLine.h
@@ -20,6 +20,9 @@ namespace UE::Text3D::Layout
 		/** Style applied on this entry */
 		FName StyleTag = NAME_None;
 
+		/** Actual size of the glyph */
+		float Size = 48.f;
+
 		/** Actual entry for the glyph */
 		FShapedGlyphEntry Entry;
 	};
diff --git a/Engine/Plugins/Experimental/Text3D/Source/Text3D/Private/LayoutBuilders/Text3DShapedGlyphText.cpp b/Engine/Plugins/Experimental/Text3D/Source/Text3D/Private/LayoutBuilders/Text3DShapedGlyphText.cpp
--- a/Engine/Plugins/Experimental/Text3D/Source/Text3D/Private/LayoutBuilders/Text3DShapedGlyphText.cpp
+++ b/Engine/Plugins/Experimental/Text3D/Source/Text3D/Private/LayoutBuilders/Text3DShapedGlyphText.cpp
@@ -92,8 +92,8 @@ namespace UE::Text3D::Layout
 	#if WITH_FREETYPE
 					const FT_Face Face = FontFaces[Index]->GetFace();
 
-					const float Ascender = FMath::Max(0.0f, Face->size->metrics.ascender * UE::Text3D::Metrics::FontSizeInverse);
-					const float Descender = FMath::Max(0.0f, -Face->size->metrics.descender * UE::Text3D::Metrics::FontSizeInverse);
+					const float Ascender = FMath::Max(0.0f, Face->size->metrics.ascender * UE::Text3D::Metrics::Convert26Dot6ToPixel);
+					const float Descender = FMath::Max(0.0f, -Face->size->metrics.descender * UE::Text3D::Metrics::Convert26Dot6ToPixel);
 
 					MaxFontAscender  = FMath::Max(MaxFontAscender,  Ascender);
 					MaxFontDescender = FMath::Max(MaxFontDescender, Descender);
diff --git a/Engine/Plugins/Experimental/Text3D/Source/Text3D/Private/Renderers/StaticMeshes/Text3DStaticMeshesRenderer.cpp b/Engine/Plugins/Experimental/Text3D/Source/Text3D/Private/Renderers/StaticMeshes/Text3DStaticMeshesRenderer.cpp
--- a/Engine/Plugins/Experimental/Text3D/Source/Text3D/Private/Renderers/StaticMeshes/Text3DStaticMeshesRenderer.cpp
+++ b/Engine/Plugins/Experimental/Text3D/Source/Text3D/Private/Renderers/StaticMeshes/Text3DStaticMeshesRenderer.cpp
@@ -162,17 +162,7 @@ void UText3DStaticMeshesRenderer::OnClear()
 
 void UText3DStaticMeshesRenderer::OnDestroy()
 {
-	FDetachmentTransformRules DetachmentRules = FDetachmentTransformRules::KeepRelativeTransform;
-	DetachmentRules.bCallModify = false;
-
-	for (UStaticMeshComponent* MeshComponent : CharacterMeshes)
-	{
-		if (MeshComponent)
-		{
-			MeshComponent->DetachFromComponent(DetachmentRules);
-			MeshComponent->DestroyComponent();
-		}
-	}
+	AllocateComponents(0);
 	CharacterMeshes.Reset();
 
 	for (UStaticMeshComponent* MeshComponent : CharacterMeshesPool)
@@ -266,6 +256,24 @@ void UText3DStaticMeshesRenderer::AllocateComponents(int32 InCount)
 	CharacterMeshes.SetNum(InCount);
 }
 
+void UText3DStaticMeshesRenderer::CollectComponents()
+{
+	TArray<UObject*> OwnedObjects;
+	GetObjectsWithOuter(this, OwnedObjects, /** Include nested */false);
+
+	for (UObject* OwnedObject : OwnedObjects)
+	{
+		if (UStaticMeshComponent* MeshComponent = Cast<UStaticMeshComponent>(OwnedObject))
+		{
+			if (!CharacterMeshes.Contains(MeshComponent)
+				&& !CharacterMeshesPool.Contains(MeshComponent))
+			{
+				CharacterMeshes.Add(MeshComponent);
+			}
+		}
+	}
+}
+
 int32 UText3DStaticMeshesRenderer::GetGlyphCount()
 {
 	if (!TextRoot)
@@ -291,6 +299,12 @@ const TArray<UStaticMeshComponent*>& UText3DStaticMeshesRenderer::GetGlyphMeshCo
 	return CharacterMeshes;
 }
 
+void UText3DStaticMeshesRenderer::PostEditImport()
+{
+	Super::PostEditImport();
+	CollectComponents();
+}
+
 #if WITH_EDITOR
 void UText3DStaticMeshesRenderer::ConvertToStaticMesh()
 {
diff --git a/Engine/Plugins/Experimental/Text3D/Source/Text3D/Private/Renderers/StaticMeshes/Text3DStaticMeshesRenderer.h b/Engine/Plugins/Experimental/Text3D/Source/Text3D/Private/Renderers/StaticMeshes/Text3DStaticMeshesRenderer.h
--- a/Engine/Plugins/Experimental/Text3D/Source/Text3D/Private/Renderers/StaticMeshes/Text3DStaticMeshesRenderer.h
+++ b/Engine/Plugins/Experimental/Text3D/Source/Text3D/Private/Renderers/StaticMeshes/Text3DStaticMeshesRenderer.h
@@ -45,6 +45,10 @@ public:
 #endif
 
 protected:
+	//~ Begin UObject
+	virtual void PostEditImport() override;
+	//~ End UObject
+
 	//~ Begin UText3DRendererBase
 	virtual void OnCreate() override;
 	virtual void OnUpdate(const UE::Text3D::Renderer::FUpdateParameters& InParameters) override;
@@ -59,6 +63,7 @@ protected:
 	//~ End UText3DRendererBase
 
 	void AllocateComponents(int32 InCount);
+	void CollectComponents();
 
 #if WITH_EDITOR
 	void SetDebugMode(bool bEnabled);
diff --git a/Engine/Plugins/Experimental/Text3D/Source/Text3D/Private/Styles/Text3DStyleBase.cpp b/Engine/Plugins/Experimental/Text3D/Source/Text3D/Private/Styles/Text3DStyleBase.cpp
--- a/Engine/Plugins/Experimental/Text3D/Source/Text3D/Private/Styles/Text3DStyleBase.cpp
+++ b/Engine/Plugins/Experimental/Text3D/Source/Text3D/Private/Styles/Text3DStyleBase.cpp
@@ -4,6 +4,7 @@
 
 #include "Engine/Font.h"
 #include "Settings/Text3DProjectSettings.h"
+#include "Styles/Text3DStyleSet.h"
 #include "Styling/StyleDefaults.h"
 #include "Text3DComponent.h"
 
@@ -21,6 +22,8 @@ void UText3DStyleBase::PostEditChangeProperty(FPropertyChangedEvent& InPropertyC
 	else if (MemberProperty == GET_MEMBER_NAME_CHECKED(UText3DStyleBase, StyleName)
 		|| MemberProperty == GET_MEMBER_NAME_CHECKED(UText3DStyleBase, bOverrideFont)
 		|| MemberProperty == GET_MEMBER_NAME_CHECKED(UText3DStyleBase, FontTypeface)
+		|| MemberProperty == GET_MEMBER_NAME_CHECKED(UText3DStyleBase, bOverrideFontSize)
+		|| MemberProperty == GET_MEMBER_NAME_CHECKED(UText3DStyleBase, FontSize)
 		|| MemberProperty == GET_MEMBER_NAME_CHECKED(UText3DStyleBase, bOverrideFrontColor)
 		|| MemberProperty == GET_MEMBER_NAME_CHECKED(UText3DStyleBase, FrontColor))
 	{
@@ -87,6 +90,29 @@ void UText3DStyleBase::SetFontTypeface(FName InTypeface)
 	OnStylePropertiesChanged();
 }
 
+void UText3DStyleBase::SetOverrideFontSize(bool bInOverride)
+{
+	if (bOverrideFontSize == bInOverride)
+	{
+		return;
+	}
+
+	bOverrideFontSize = bInOverride;
+	OnStylePropertiesChanged();
+}
+
+void UText3DStyleBase::SetFontSize(float InSize)
+{
+	InSize = FMath::Max(InSize, 0.1f);
+	if (FMath::IsNearlyEqual(FontSize, InSize))
+	{
+		return;
+	}
+
+	FontSize = InSize;
+	OnStylePropertiesChanged();
+}
+
 void UText3DStyleBase::SetOverrideFrontColor(bool bInOverride)
 {
 	if (bOverrideFrontColor == bInOverride)
@@ -113,13 +139,18 @@ void UText3DStyleBase::ConfigureStyle(FTextBlockStyle& OutStyle) const
 {
 	if (bOverrideFont)
 	{
-		FSlateFontInfo FontInfo(Font, UE::Text3D::Metrics::SlateFontSize);
+		FSlateFontInfo FontInfo(Font, OutStyle.Font.Size);
 		FontInfo.CompositeFont = Font && Font->GetCompositeFont() ? MakeShared<FCompositeFont>(*Font->GetCompositeFont()) : FStyleDefaults::GetFontInfo().CompositeFont;
 		FontInfo.TypefaceFontName = FontTypeface;
 		FontInfo.FontFallback = EFontFallback::FF_NoFallback;
 		OutStyle.SetFont(FontInfo);
 	}
 
+	if (bOverrideFontSize)
+	{
+		OutStyle.SetFontSize(FontSize);
+	}
+
 	if (bOverrideFrontColor)
 	{
 		OutStyle.SetColorAndOpacity(FrontColor);
@@ -134,6 +165,10 @@ void UText3DStyleBase::OnStylePropertiesChanged()
 		{
 			Text3DComponent->RequestUpdate(EText3DRendererFlags::All, /** Immediate */false);
 		}
+		else if (UText3DStyleSet* StyleSet = GetTypedOuter<UText3DStyleSet>())
+		{
+			StyleSet->OnStyleSetPropertiesChanged();
+		}
 	}
 }
 
diff --git a/Engine/Plugins/Experimental/Text3D/Source/Text3D/Private/Subsystems/Text3DEngineSubsystem.cpp b/Engine/Plugins/Experimental/Text3D/Source/Text3D/Private/Subsystems/Text3DEngineSubsystem.cpp
--- a/Engine/Plugins/Experimental/Text3D/Source/Text3D/Private/Subsystems/Text3DEngineSubsystem.cpp
+++ b/Engine/Plugins/Experimental/Text3D/Source/Text3D/Private/Subsystems/Text3DEngineSubsystem.cpp
@@ -217,6 +217,11 @@ FText3DCachedMesh* FText3DFontFaceCache::FindOrBuildGlyphMesh(uint32 InGlyphInde
 
 	FText3DCachedMesh& CachedMesh = GlyphMeshes.FindOrAdd(HashParameters);
 
+	if (const FT_Face Face = FontFace->GetFace())
+	{
+		CachedMesh.FontFaceGlyphSize = FVector2D(Face->size->metrics.x_ppem, Face->size->metrics.y_ppem);
+	}
+
 	const FString MeshName = FString::Printf(TEXT("Text3D_Char_%u"), HashParameters);
 	const FName StaticMeshName = MakeUniqueObjectName(TextSubsystem, UStaticMesh::StaticClass(), FName(*MeshName));
 	CachedMesh.StaticMesh = NewObject<UStaticMesh>(TextSubsystem, StaticMeshName, RF_Transient | RF_Public);
diff --git a/Engine/Plugins/Experimental/Text3D/Source/Text3D/Private/Text3DComponent.cpp b/Engine/Plugins/Experimental/Text3D/Source/Text3D/Private/Text3DComponent.cpp
--- a/Engine/Plugins/Experimental/Text3D/Source/Text3D/Private/Text3DComponent.cpp
+++ b/Engine/Plugins/Experimental/Text3D/Source/Text3D/Private/Text3DComponent.cpp
@@ -97,7 +97,7 @@ void UText3DComponent::PostLoad()
 
 void UText3DComponent::RequestUpdate(EText3DRendererFlags InFlags, bool bInImmediate)
 {
-	if (InFlags != EText3DRendererFlags::None)
+	if (InFlags != EText3DRendererFlags::None && !IsTemplate(RF_ClassDefaultObject))
 	{
 		if (bInImmediate)
 		{
@@ -342,31 +342,7 @@ void UText3DComponent::Serialize(FArchive& InArchive)
 void UText3DComponent::PostEditImport()
 {
 	Super::PostEditImport();
-
-	// Clear renderer after duplication to avoid duplicated characters
-	TArray<UObject*> OwnedObjects;
-	GetObjectsWithOuter(this, OwnedObjects, false);
-
-	for (UObject* OwnedObject : OwnedObjects)
-	{
-		if (UText3DRendererBase* Renderer = Cast<UText3DRendererBase>(OwnedObject))
-		{
-			if (Renderer != TextRenderer)
-			{
-				Renderer->Clear();
-				Renderer->Destroy();
-				Renderer->MarkAsGarbage();
-			}
-		}
-	}
-
-	RequestUpdate(EText3DRendererFlags::All);
-}
-
-void UText3DComponent::PostDuplicate(EDuplicateMode::Type InDuplicateMode)
-{
-	Super::PostDuplicate(InDuplicateMode);
-
+	CleanupRenderer();
 	RequestUpdate(EText3DRendererFlags::All);
 }
 
@@ -386,14 +362,15 @@ void UText3DComponent::PostEditChangeProperty(FPropertyChangedEvent& PropertyCha
 	static const TSet<FName> TextPropertyNames
 	{
 		GET_MEMBER_NAME_CHECKED(UText3DComponent, Text),
-		GET_MEMBER_NAME_CHECKED(UText3DComponent, bEnforceUpperCase)
+		GET_MEMBER_NAME_CHECKED(UText3DComponent, bEnforceUpperCase),
+		GET_MEMBER_NAME_CHECKED(UText3DComponent, FontSize)
 	};
 
 	if (MemberPropertyName == GET_MEMBER_NAME_CHECKED(UText3DComponent, TextRendererClass))
 	{
 		OnTextRendererClassChanged();
 	}
-	if (TextPropertyNames.Contains(MemberPropertyName))
+	else if (TextPropertyNames.Contains(MemberPropertyName))
 	{
 		OnTextChanged();
 	}
@@ -489,6 +466,23 @@ void UText3DComponent::SetFont(UFont* InFont)
 	OnFontPropertiesChanged();
 }
 
+void UText3DComponent::SetFontSize(float InSize)
+{
+	InSize = FMath::Max(0.1, InSize);
+	if (FMath::IsNearlyEqual(InSize, FontSize))
+	{
+		return;
+	}
+
+	FontSize = InSize;
+	OnFontPropertiesChanged();
+}
+
+float UText3DComponent::GetFontSize() const
+{
+	return FontSize;
+}
+
 void UText3DComponent::SetHasOutline(const bool bValue)
 {
 	if (UText3DDefaultGeometryExtension* DefaultGeometryExtension = GetCastedGeometryExtension<UText3DDefaultGeometryExtension>())
@@ -1011,10 +1005,15 @@ void UText3DComponent::Rebuild(EText3DRendererFlags InUpdateFlags)
 		return;
 	}
 
-	if (!TextRenderer)
+	if (!TextRenderer || TextRenderer->GetClass() != TextRendererClass)
 	{
-		UE_LOG(LogText3D, Error, TEXT("Text3D renderer is not valid or initialized, cannot proceed"));
-		return;
+		OnTextRendererClassChanged();
+
+		if (!TextRenderer)
+		{
+			UE_LOG(LogText3D, Error, TEXT("Text3D renderer is not valid, cannot proceed"));
+			return;
+		}
 	}
 
 	if (!IsRegistered())
@@ -1098,6 +1097,11 @@ void UText3DComponent::Rebuild(EText3DRendererFlags InUpdateFlags)
 
 void UText3DComponent::OnTextRendererClassChanged()
 {
+	if (IsTemplate(RF_ClassDefaultObject))
+	{
+		return;
+	}
+
 	const UText3DRendererBase* RendererCDO = TextRendererClass.Get() ? TextRendererClass->GetDefaultObject<UText3DRendererBase>() : nullptr;
 
 	if (!RendererCDO)
@@ -1127,7 +1131,12 @@ void UText3DComponent::OnTextRendererClassChanged()
 		TextRenderer = nullptr;
 	}
 
-	TextRenderer = NewObject<UText3DRendererBase>(this, TextRendererClass.Get(), RendererCDO->GetFriendlyName(), RF_Public);
+	TextRenderer = FindObject<UText3DRendererBase>(this, RendererCDO->GetFriendlyName().ToString());
+
+	if (!TextRenderer)
+	{
+		TextRenderer = NewObject<UText3DRendererBase>(this, TextRendererClass.Get(), RendererCDO->GetFriendlyName(), RF_Public);
+	}
 
 	const AActor* ContextActor = GetOwner();
 	UE_LOG(LogText3D, Log, TEXT("%s : Text3D renderer changed : old %s - new %s"), !!ContextActor ? *ContextActor->GetActorNameOrLabel() : TEXT("Invalid owner"), *OldRendererName.ToString(), *TextRenderer->GetFriendlyName().ToString());
@@ -1152,6 +1161,33 @@ void UText3DComponent::OnFontPropertiesChanged()
 	RequestUpdate(EText3DRendererFlags::All);
 }
 
+void UText3DComponent::CleanupRenderer()
+{
+	TArray<UObject*> OwnedObjects;
+	GetObjectsWithOuter(this, OwnedObjects, /** Include nested */false);
+
+	for (UObject* OwnedObject : OwnedObjects)
+	{
+		if (UText3DRendererBase* Renderer = Cast<UText3DRendererBase>(OwnedObject))
+		{
+			if (!TextRenderer && Renderer->GetClass() == TextRendererClass)
+			{
+				TextRenderer = Renderer;
+			}
+			else if (Renderer != TextRenderer)
+			{
+				Renderer->Destroy();
+				Renderer->MarkAsGarbage();
+			}
+		}
+	}
+
+	if (TextRenderer)
+	{
+		TextRenderer->Clear();
+	}
+}
+
 TArray<FName> UText3DComponent::GetTypefaceNames() const
 {
 	TArray<FName> TypefaceNames;
diff --git a/Engine/Plugins/Experimental/Text3D/Source/Text3D/Public/Extensions/Text3DDefaultGeometryExtension.h b/Engine/Plugins/Experimental/Text3D/Source/Text3D/Public/Extensions/Text3DDefaultGeometryExtension.h
--- a/Engine/Plugins/Experimental/Text3D/Source/Text3D/Public/Extensions/Text3DDefaultGeometryExtension.h
+++ b/Engine/Plugins/Experimental/Text3D/Source/Text3D/Public/Extensions/Text3DDefaultGeometryExtension.h
@@ -106,6 +106,7 @@ public:
 protected:
 	//~ Begin UObject
 #if WITH_EDITOR
+	virtual void PostEditUndo() override;
 	virtual void PostEditChangeProperty(FPropertyChangedEvent& InEvent) override;
 #endif
 	//~ End UObject
diff --git a/Engine/Plugins/Experimental/Text3D/Source/Text3D/Public/Extensions/Text3DDefaultRenderingExtension.h b/Engine/Plugins/Experimental/Text3D/Source/Text3D/Public/Extensions/Text3DDefaultRenderingExtension.h
--- a/Engine/Plugins/Experimental/Text3D/Source/Text3D/Public/Extensions/Text3DDefaultRenderingExtension.h
+++ b/Engine/Plugins/Experimental/Text3D/Source/Text3D/Public/Extensions/Text3DDefaultRenderingExtension.h
@@ -71,6 +71,7 @@ public:
 protected:
 	//~ Begin UObject
 #if WITH_EDITOR
+	virtual void PostEditUndo() override;
 	virtual void PostEditChangeProperty(FPropertyChangedEvent& InEvent) override;
 #endif
 	//~ End UObject
diff --git a/Engine/Plugins/Experimental/Text3D/Source/Text3D/Public/Extensions/Text3DDefaultTokenExtension.h b/Engine/Plugins/Experimental/Text3D/Source/Text3D/Public/Extensions/Text3DDefaultTokenExtension.h
--- a/Engine/Plugins/Experimental/Text3D/Source/Text3D/Public/Extensions/Text3DDefaultTokenExtension.h
+++ b/Engine/Plugins/Experimental/Text3D/Source/Text3D/Public/Extensions/Text3DDefaultTokenExtension.h
@@ -21,9 +21,9 @@ protected:
 #endif
 	//~ End UObject
 
-	//~ Begin UText3DTemplateExtensionBase
+	//~ Begin UText3DTokenExtensionBase
 	virtual const FText& GetFormattedText() const override;
-	//~ End UText3DTemplateExtensionBase
+	//~ End UText3DTokenExtensionBase
 
 	//~ Begin UText3DExtensionBase
 	virtual EText3DExtensionResult PreRendererUpdate(const UE::Text3D::Renderer::FUpdateParameters& InParameters) override;
diff --git a/Engine/Plugins/Experimental/Text3D/Source/Text3D/Public/Extensions/Text3DLayoutTransformEffect.h b/Engine/Plugins/Experimental/Text3D/Source/Text3D/Public/Extensions/Text3DLayoutTransformEffect.h
--- a/Engine/Plugins/Experimental/Text3D/Source/Text3D/Public/Extensions/Text3DLayoutTransformEffect.h
+++ b/Engine/Plugins/Experimental/Text3D/Source/Text3D/Public/Extensions/Text3DLayoutTransformEffect.h
@@ -165,6 +165,7 @@ protected:
 
 	//~ Begin UObject
 #if WITH_EDITOR
+	virtual void PostEditUndo() override;
 	virtual void PostEditChangeProperty(FPropertyChangedEvent& InEvent) override;
 #endif
 	//~ End UObject
diff --git a/Engine/Plugins/Experimental/Text3D/Source/Text3D/Public/Styles/Text3DStyleBase.h b/Engine/Plugins/Experimental/Text3D/Source/Text3D/Public/Styles/Text3DStyleBase.h
--- a/Engine/Plugins/Experimental/Text3D/Source/Text3D/Public/Styles/Text3DStyleBase.h
+++ b/Engine/Plugins/Experimental/Text3D/Source/Text3D/Public/Styles/Text3DStyleBase.h
@@ -57,6 +57,24 @@ public:
 		return FontTypeface;
 	}
 
+	UFUNCTION(BlueprintCallable, Category = "Text3D|Style")
+	TEXT3D_API void SetOverrideFontSize(bool bInOverride);
+
+	UFUNCTION(BlueprintPure, Category = "Text3D|Style")
+	bool GetOverrideFontSize() const
+	{
+		return bOverrideFontSize;
+	}
+
+	UFUNCTION(BlueprintCallable, Category = "Text3D|Style")
+	TEXT3D_API void SetFontSize(float InSize);
+
+	UFUNCTION(BlueprintPure, Category = "Text3D|Style")
+	float GetFontSize() const
+	{
+		return FontSize;
+	}
+
 	UFUNCTION(BlueprintCallable, Category = "Text3D|Style")
 	TEXT3D_API void SetOverrideFrontColor(bool bInOverride);
 
@@ -108,11 +126,21 @@ protected:
 	UPROPERTY(EditAnywhere, Setter, Getter, Category="Style", meta=(EditCondition="bOverrideFont", GetOptions="GetTypefaceNames", AllowPrivateAccess="true"))
 	FName FontTypeface;
 
+	/** Behavior for the size rule of this style */
+	UPROPERTY(EditAnywhere, Setter="SetOverrideFontSize", Getter="GetOverrideFontSize", Category="Style", meta=(AllowPrivateAccess="true", InlineEditConditionToggle))
+	bool bOverrideFontSize = true;
+
+	/** Text font size to render and scale each character */
+	UPROPERTY(EditAnywhere, Setter, Getter, Category="Style", meta=(ClampMin="0.1", EditCondition="bOverrideFontSize", AllowPrivateAccess="true"))
+	float FontSize = 48.f;
+
 	/** Behavior for the front color rule of this style */
 	UPROPERTY(EditAnywhere, Setter="SetOverrideFrontColor", Getter="GetOverrideFrontColor", Category="Style", meta=(AllowPrivateAccess="true", InlineEditConditionToggle))
 	bool bOverrideFrontColor = true;
 
 	/** Text front color */
-	UPROPERTY(EditAnywhere, Category="Style", meta=(EditCondition="bOverrideFrontColor", AllowPrivateAccess="true"))
+	UPROPERTY(EditAnywhere, Setter, Getter, Category="Style", meta=(EditCondition="bOverrideFrontColor", AllowPrivateAccess="true"))
 	FLinearColor FrontColor = FLinearColor::White;
+
+	
 };
diff --git a/Engine/Plugins/Experimental/Text3D/Source/Text3D/Public/Styles/Text3DStyleSet.h b/Engine/Plugins/Experimental/Text3D/Source/Text3D/Public/Styles/Text3DStyleSet.h
--- a/Engine/Plugins/Experimental/Text3D/Source/Text3D/Public/Styles/Text3DStyleSet.h
+++ b/Engine/Plugins/Experimental/Text3D/Source/Text3D/Public/Styles/Text3DStyleSet.h
@@ -13,6 +13,8 @@ class UText3DStyleSet : public UObject
 {
 	GENERATED_BODY()
 
+	friend class UText3DStyleBase;
+
 public:
 	DECLARE_MULTICAST_DELEGATE_OneParam(FOnStyleSetUpdated, UText3DStyleSet* /** InStyleSet */)
 
diff --git a/Engine/Plugins/Experimental/Text3D/Source/Text3D/Public/Text3DComponent.h b/Engine/Plugins/Experimental/Text3D/Source/Text3D/Public/Text3DComponent.h
--- a/Engine/Plugins/Experimental/Text3D/Source/Text3D/Public/Text3DComponent.h
+++ b/Engine/Plugins/Experimental/Text3D/Source/Text3D/Public/Text3DComponent.h
@@ -98,6 +98,13 @@ public:
 	UFUNCTION(BlueprintCallable, Category = "Rendering|Components|Text3D")
 	TEXT3D_API void SetFont(UFont* InFont);
 
+	/** Set the text font size */
+	UFUNCTION(BlueprintCallable, Category = "Rendering|Components|Text3D")
+	TEXT3D_API void SetFontSize(float InSize);
+
+	UFUNCTION(BlueprintPure, Category = "Rendering|Components|Text3D")
+	TEXT3D_API float GetFontSize() const;
+
 	/** Set whether an outline is applied. */
 	UFUNCTION(BlueprintCallable, Category = "Rendering|Components|Text3D")
 	TEXT3D_API void SetHasOutline(const bool bValue);
@@ -415,7 +422,6 @@ protected:
 	TEXT3D_API virtual void Serialize(FArchive& InArchive) override;
 	TEXT3D_API virtual void PostLoad() override;
 	TEXT3D_API virtual void PostEditImport() override;
-	TEXT3D_API virtual void PostDuplicate(EDuplicateMode::Type InDuplicateMode) override;
 #if WITH_EDITOR
 	TEXT3D_API virtual void PostEditChangeProperty(FPropertyChangedEvent& PropertyChangedEvent) override;
 	TEXT3D_API virtual void PostTransacted(const FTransactionObjectEvent& TransactionEvent) override;
@@ -460,6 +466,9 @@ protected:
 	/** Called when font options are changed */
 	void OnFontPropertiesChanged();
 
+	/** Clear renderer after duplication to avoid duplicated characters */
+	void CleanupRenderer();
+
 	/** The renderer class to use to create the text geometry */
 	UPROPERTY(EditAnywhere, Setter, Getter, Category="Text", NoClear, meta=(ShowDisplayNames, AllowPrivateAccess="true"))
 	TSubclassOf<UText3DRendererBase> TextRendererClass;
@@ -476,6 +485,10 @@ protected:
 	UPROPERTY(EditAnywhere, Getter, Setter, DisplayName = "Default Typeface", Category = "Text", meta = (AllowPrivateAccess = "true", GetOptions="GetTypefaceNames"))
 	FName Typeface;
 
+	/** Text font size defines the size of each rendered character */
+	UPROPERTY(EditAnywhere, Getter, Setter, Category = "Text", meta=(ClampMin = "0.1", AllowPrivateAccess = "true"))
+	float FontSize = 48.f;
+
 	/** Whether to force upper case for text */
 	UPROPERTY(EditAnywhere, Getter = "GetEnforceUpperCase", Setter = "SetEnforceUpperCase", Category = "Text", meta = (AllowPrivateAccess="true"))
 	bool bEnforceUpperCase = false;
diff --git a/Engine/Plugins/Experimental/Text3D/Source/Text3D/Public/Text3DTypes.h b/Engine/Plugins/Experimental/Text3D/Source/Text3D/Public/Text3DTypes.h
--- a/Engine/Plugins/Experimental/Text3D/Source/Text3D/Public/Text3DTypes.h
+++ b/Engine/Plugins/Experimental/Text3D/Source/Text3D/Public/Text3DTypes.h
@@ -293,9 +293,14 @@ struct FText3DCachedMesh
 	UPROPERTY()
 	FBox MeshBounds = FBox(ForceInitToZero);
 
+	/** Pivot offset */
 	UPROPERTY()
 	FVector MeshOffset = FVector::ZeroVector;
 
+	/** Used to scale mesh to achieve different font size */
+	UPROPERTY()
+	FVector2D FontFaceGlyphSize = FVector2D::ZeroVector;
+
 	/** Amount of active usage of this glyph, used for cleanup */
 	uint32 RefCount = 0;
 };
@@ -343,14 +348,12 @@ namespace UE::Text3D
 
 	namespace Metrics
 	{
-		/** Original value to match size across different Text3D version */
-		static constexpr float FontSize = 64;
-		/** Value used to match size across Text3D implementation */
-		static constexpr float SlateFontSize = 48;
-		/** Scale used to transform freetype face result to get normalized values */
-		static constexpr float FontSizeInverse = 1.0f / FontSize;
-		/** DPI used to match slate */
-		static constexpr float FontDPI = 96.f;
+		/** Default value to match size across different Text3D version */
+		static constexpr float FontSize = 48.f;
+		/** Scale used to transform freetype face result to get normalized values based on our font size */
+		static constexpr float FontSizeInverse = 1.f / FontSize;
+		/** Freetype conversion ratio to pixel */
+		static constexpr float Convert26Dot6ToPixel = 1.f / 64.f;
 	}
 
 	namespace Renderer
