MotionDesign : Text3D
- Moving Text3D plugin outside of Experimental and into VirtualProduction

#virtualized

diff --git a/Engine/Plugins/VirtualProduction/Text3D/Source/Text3D/Private/Extensions/Text3DExtensionBase.cpp b/Engine/Plugins/VirtualProduction/Text3D/Source/Text3D/Private/Extensions/Text3DExtensionBase.cpp
new file mode 100644
--- /dev/null
+++ b/Engine/Plugins/VirtualProduction/Text3D/Source/Text3D/Private/Extensions/Text3DExtensionBase.cpp
@@ -0,0 +1,18 @@
+// Copyright Epic Games, Inc. All Rights Reserved.
+
+#include "Extensions/Text3DExtensionBase.h"
+
+#include "Text3DComponent.h"
+
+void UText3DExtensionBase::RequestUpdate(EText3DRendererFlags InFlags, bool bInImmediate) const
+{
+	if (UText3DComponent* Text3DComponent = GetText3DComponent())
+	{
+		Text3DComponent->RequestUpdate(InFlags, bInImmediate);
+	}
+}
+
+UText3DComponent* UText3DExtensionBase::GetText3DComponent() const
+{
+	return GetTypedOuter<UText3DComponent>();
+}
diff --git a/Engine/Plugins/VirtualProduction/Text3D/Source/Text3D/Private/Renderers/Text3DRendererBase.cpp b/Engine/Plugins/VirtualProduction/Text3D/Source/Text3D/Private/Renderers/Text3DRendererBase.cpp
new file mode 100644
--- /dev/null
+++ b/Engine/Plugins/VirtualProduction/Text3D/Source/Text3D/Private/Renderers/Text3DRendererBase.cpp
@@ -0,0 +1,252 @@
+// Copyright Epic Games, Inc. All Rights Reserved.
+
+#include "Renderers/Text3DRendererBase.h"
+
+#include "Extensions/Text3DCharacterExtensionBase.h"
+#include "Extensions/Text3DGeometryExtensionBase.h"
+#include "Extensions/Text3DLayoutEffectBase.h"
+#include "Extensions/Text3DLayoutExtensionBase.h"
+#include "Extensions/Text3DMaterialExtensionBase.h"
+#include "Extensions/Text3DRenderingExtensionBase.h"
+#include "Extensions/Text3DStyleExtensionBase.h"
+#include "Extensions/Text3DTokenExtensionBase.h"
+#include "GameFramework/Actor.h"
+#include "Logs/Text3DLogs.h"
+#include "Settings/Text3DProjectSettings.h"
+#include "Text3DComponent.h"
+
+#if WITH_EDITOR
+void UText3DRendererBase::BindDebugDelegate()
+{
+	if (UText3DProjectSettings* Text3DSettings = UText3DProjectSettings::GetMutable())
+	{
+		Text3DSettings->OnSettingChanged().AddUObject(this, &UText3DRendererBase::OnTextSettingsChanged);
+	}
+}
+
+void UText3DRendererBase::UnbindDebugDelegate()
+{
+	if (UText3DProjectSettings* Text3DSettings = UText3DProjectSettings::GetMutable())
+	{
+		Text3DSettings->OnSettingChanged().RemoveAll(this);
+	}
+}
+
+void UText3DRendererBase::OnTextSettingsChanged(UObject* InSettings, FPropertyChangedEvent& InEvent)
+{
+	if (InEvent.GetMemberPropertyName() == UText3DProjectSettings::GetDebugModePropertyName())
+	{
+		OnDebugModeChanged();
+	}
+}
+
+void UText3DRendererBase::OnDebugModeChanged()
+{
+	if (const UText3DProjectSettings* TextSettings = UText3DProjectSettings::Get())
+	{
+		if (TextSettings->GetDebugMode())
+		{
+			OnDebugModeEnabled();
+		}
+		else
+		{
+			OnDebugModeDisabled();
+		}
+	}
+}
+#endif
+
+UText3DComponent* UText3DRendererBase::GetText3DComponent() const
+{
+	return GetTypedOuter<UText3DComponent>();
+}
+
+void UText3DRendererBase::RefreshBounds()
+{
+	CachedBounds = OnCalculateBounds();
+}
+
+void UText3DRendererBase::Create()
+{
+	if (bInitialized)
+	{
+		return;
+	}
+
+	const UText3DComponent* Text3DComponent = GetText3DComponent();
+	if (!IsValid(Text3DComponent))
+	{
+		return;
+	}
+
+	OnCreate();
+
+#if WITH_EDITOR
+	BindDebugDelegate();
+	OnDebugModeChanged();
+#endif
+
+	bInitialized = true;
+
+	const AActor* Owner = Text3DComponent->GetOwner();
+	UE_LOG(LogText3D, Log, TEXT("%s : Text3DRenderer %s Created"), !!Owner ? *Owner->GetActorNameOrLabel() : TEXT("Invalid owner"), *GetFriendlyName().ToString())
+}
+
+void UText3DRendererBase::Update(EText3DRendererFlags InFlags)
+{
+	using namespace UE::Text3D::Renderer;
+
+	if (!bInitialized)
+	{
+		return;
+	}
+
+	const UText3DComponent* Text3DComponent = GetText3DComponent();
+	if (!IsValid(Text3DComponent))
+	{
+		return;
+	}
+
+	TArray<UText3DExtensionBase*, TInlineAllocator<6>> Extensions
+	{
+		Text3DComponent->GetTokenExtension(),
+		Text3DComponent->GetStyleExtension(),
+		Text3DComponent->GetCharacterExtension(),
+		Text3DComponent->GetGeometryExtension(),
+		Text3DComponent->GetLayoutExtension(),
+		Text3DComponent->GetMaterialExtension(),
+		Text3DComponent->GetRenderingExtension()
+	};
+
+	Extensions.Append(Text3DComponent->GetLayoutEffects());
+
+	Extensions.RemoveAll([](const UText3DExtensionBase* InExtension)
+	{
+		return !IsValid(InExtension);
+	});
+
+	// Sort in reverse order since we iterate in reverse order
+	Extensions.StableSort([](const UText3DExtensionBase& InExtensionA, const UText3DExtensionBase& InExtensionB)->bool
+	{
+		return InExtensionA.GetUpdatePriority() > InExtensionB.GetUpdatePriority();
+	});
+
+	uint8 Flag = 1 << 0;
+	uint8 FlagsLeft = static_cast<uint8>(InFlags);
+
+	while (Flag < static_cast<uint8>(EText3DRendererFlags::All))
+	{
+		const EText3DRendererFlags FlagCasted = static_cast<EText3DRendererFlags>(Flag);
+
+		if (EnumHasAnyFlags(InFlags, FlagCasted))
+		{
+			FlagsLeft -= Flag;
+
+			const FUpdateParameters Parameters
+			{
+				.UpdateFlags = InFlags,
+				.CurrentFlag = FlagCasted,
+				.bIsLastFlag = FlagsLeft == 0
+			};
+
+			for (int32 Index = Extensions.Num() - 1; Index >= 0; --Index)
+			{
+				if (UText3DExtensionBase* Extension = Extensions[Index])
+				{
+					const EText3DExtensionResult ExtensionStatus = Extension->PreRendererUpdate(Parameters);
+
+					if (ExtensionStatus == EText3DExtensionResult::Failed)
+					{
+						UE_LOG(LogText3D, Error, TEXT("Failed to PRE update Text3D %s extension"), *GetNameSafe(Extension->GetClass()))
+						return;
+					}
+					else if (ExtensionStatus == EText3DExtensionResult::Finished)
+					{
+						Extensions.RemoveAt(Index);
+					}
+				}
+			}
+
+			OnUpdate(Parameters);
+
+			for (int32 Index = Extensions.Num() - 1; Index >= 0; --Index)
+			{
+				if (UText3DExtensionBase* Extension = Extensions[Index])
+				{
+					const EText3DExtensionResult ExtensionStatus = Extension->PostRendererUpdate(Parameters);
+
+					if (ExtensionStatus == EText3DExtensionResult::Failed)
+					{
+						UE_LOG(LogText3D, Error, TEXT("Failed to POST update Text3D %s extension"), *GetNameSafe(Extension->GetClass()))
+						return;
+					}
+					else if (ExtensionStatus == EText3DExtensionResult::Finished)
+					{
+						Extensions.RemoveAt(Index);
+					}
+				}
+			}
+		}
+
+		Flag <<= 1;
+	}
+
+	const AActor* Owner = Text3DComponent->GetOwner();
+	UE_LOG(LogText3D, Verbose, TEXT("%s : Text3DRenderer %s Updated with flags %i"), !!Owner ? *Owner->GetActorNameOrLabel() : TEXT("Invalid owner"), *GetFriendlyName().ToString(), InFlags)
+}
+
+void UText3DRendererBase::Clear()
+{
+	if (!bInitialized)
+	{
+		return;
+	}
+
+	const UText3DComponent* Text3DComponent = GetText3DComponent();
+	if (!Text3DComponent)
+	{
+		return;
+	}
+
+	OnClear();
+	CachedBounds.Reset();
+
+	const AActor* Owner = Text3DComponent->GetOwner();
+	UE_LOG(LogText3D, Verbose, TEXT("%s : Text3DRenderer %s Cleared"), !!Owner ? *Owner->GetActorNameOrLabel() : TEXT("Invalid owner"), *GetFriendlyName().ToString())
+}
+
+void UText3DRendererBase::Destroy()
+{
+	if (!bInitialized)
+	{
+		return;
+	}
+
+	const UText3DComponent* Text3DComponent = GetText3DComponent();
+	if (!Text3DComponent)
+	{
+		return;
+	}
+
+	OnDestroy();
+	CachedBounds.Reset();
+
+#if WITH_EDITOR
+	UnbindDebugDelegate();
+#endif
+
+	bInitialized = false;
+
+	const AActor* Owner = Text3DComponent->GetOwner();
+	UE_LOG(LogText3D, Log, TEXT("%s : Text3DRenderer %s Destroyed"), !!Owner ? *Owner->GetActorNameOrLabel() : TEXT("Invalid owner"), *GetFriendlyName().ToString())
+}
+
+FBox UText3DRendererBase::GetBounds() const
+{
+	if (!CachedBounds.IsSet())
+	{
+		return FBox(ForceInitToZero);
+	}
+
+	return CachedBounds.GetValue();
+}
\ No newline at end of file
