MotionDesign : Text3D
- Fixed spawnable behavior with sequencer and Scene State Task
- Fixed material invalidated and not recreated causing world grid material to be used
- Fixed text component used in a blueprint and spawned in a level does not update when properties are changed

diff --git a/Engine/Plugins/VirtualProduction/Text3D/Source/Text3D/Private/Characters/Text3DCharacterBase.cpp b/Engine/Plugins/VirtualProduction/Text3D/Source/Text3D/Private/Characters/Text3DCharacterBase.cpp
--- a/Engine/Plugins/VirtualProduction/Text3D/Source/Text3D/Private/Characters/Text3DCharacterBase.cpp
+++ b/Engine/Plugins/VirtualProduction/Text3D/Source/Text3D/Private/Characters/Text3DCharacterBase.cpp
@@ -128,6 +128,7 @@ void UText3DCharacterBase::ResetCharacterState()
 	RelativeScale = CDO->RelativeScale;
 	bVisible = CDO->bVisible;
 	Transform = CDO->Transform;
+	StyleTag = CDO->StyleTag;
 	FontFaceGlyphHandle.Unset();
 }
 
diff --git a/Engine/Plugins/VirtualProduction/Text3D/Source/Text3D/Private/Extensions/Text3DDefaultMaterialExtension.cpp b/Engine/Plugins/VirtualProduction/Text3D/Source/Text3D/Private/Extensions/Text3DDefaultMaterialExtension.cpp
--- a/Engine/Plugins/VirtualProduction/Text3D/Source/Text3D/Private/Extensions/Text3DDefaultMaterialExtension.cpp
+++ b/Engine/Plugins/VirtualProduction/Text3D/Source/Text3D/Private/Extensions/Text3DDefaultMaterialExtension.cpp
@@ -311,6 +311,7 @@ EText3DExtensionResult UText3DDefaultMaterialExtension::PreRendererUpdate(const
 
 		switch (Style)
 		{
+		default:	
 		case EText3DMaterialStyle::Solid:
 			DynFrontMaterial->SetVectorParameterValue(UText3DProjectSettings::FMaterialParameters::SolidColor, FrontColor);
 			DynBackMaterial->SetVectorParameterValue(UText3DProjectSettings::FMaterialParameters::SolidColor, BackColor);
@@ -332,8 +333,6 @@ EText3DExtensionResult UText3DDefaultMaterialExtension::PreRendererUpdate(const
 			SetTextureParameter(Materials, UText3DProjectSettings::FMaterialParameters::MainTexture, TextureAsset);
 			SetScalarParameter(Materials, UText3DProjectSettings::FMaterialParameters::Mode, 2);
 			break;
-		default:
-			break;
 		}
 
 		if (BlendMode == EText3DMaterialBlendMode::Translucent || Style == EText3DMaterialStyle::Gradient)
@@ -539,7 +538,10 @@ UMaterialInstanceDynamic* UText3DDefaultMaterialExtension::FindOrAdd(EText3DGrou
 	FText3DMaterialGroupKey GroupKey(MaterialKey, InGroup, InTag);
 	if (const TObjectPtr<UMaterialInstanceDynamic>* Material = GroupDynamicMaterials.Find(GroupKey))
 	{
-		return *Material;
+		if (*Material)
+		{
+			return *Material;
+		}
 	}
 
 	UMaterialInstanceDynamic* DynamicMaterial = UMaterialInstanceDynamic::Create(ParentMaterial, this);
@@ -714,6 +716,14 @@ void UText3DDefaultMaterialExtension::PostLoad()
 	PreCacheMaterials();
 }
 
+void UText3DDefaultMaterialExtension::PostInitProperties()
+{
+	Super::PostInitProperties();
+
+	GroupDynamicMaterials.Empty();
+	PreCacheMaterials();
+}
+
 #if WITH_EDITOR
 void UText3DDefaultMaterialExtension::PostEditChangeProperty(FPropertyChangedEvent& InEvent)
 {
diff --git a/Engine/Plugins/VirtualProduction/Text3D/Source/Text3D/Private/Renderers/DynamicMesh/Text3DDynamicMeshRenderer.h b/Engine/Plugins/VirtualProduction/Text3D/Source/Text3D/Private/Renderers/DynamicMesh/Text3DDynamicMeshRenderer.h
--- a/Engine/Plugins/VirtualProduction/Text3D/Source/Text3D/Private/Renderers/DynamicMesh/Text3DDynamicMeshRenderer.h
+++ b/Engine/Plugins/VirtualProduction/Text3D/Source/Text3D/Private/Renderers/DynamicMesh/Text3DDynamicMeshRenderer.h
@@ -81,7 +81,7 @@ protected:
 	void SaveMesh(const FScopedMeshEditor& InMeshEditor);
 	void RestoreMesh(FScopedMeshEditor& InMeshEditor);
 
-	UPROPERTY(VisibleAnywhere, Category="Renderer", meta=(EditCondition="false", EditConditionHides))
+	UPROPERTY(VisibleAnywhere, DuplicateTransient, TextExportTransient, Category="Renderer", meta=(EditCondition="false", EditConditionHides))
 	TObjectPtr<UDynamicMeshComponent> DynamicMeshComponent;
 
 	TArray<FGlyphMeshData> GlyphMeshData;
diff --git a/Engine/Plugins/VirtualProduction/Text3D/Source/Text3D/Private/Renderers/StaticMeshes/Text3DStaticMeshesRenderer.cpp b/Engine/Plugins/VirtualProduction/Text3D/Source/Text3D/Private/Renderers/StaticMeshes/Text3DStaticMeshesRenderer.cpp
--- a/Engine/Plugins/VirtualProduction/Text3D/Source/Text3D/Private/Renderers/StaticMeshes/Text3DStaticMeshesRenderer.cpp
+++ b/Engine/Plugins/VirtualProduction/Text3D/Source/Text3D/Private/Renderers/StaticMeshes/Text3DStaticMeshesRenderer.cpp
@@ -11,7 +11,6 @@
 #include "Extensions/Text3DMaterialExtensionBase.h"
 #include "Extensions/Text3DRenderingExtensionBase.h"
 #include "Logs/Text3DLogs.h"
-#include "Subsystems/Text3DEngineSubsystem.h"
 #include "Text3DComponent.h"
 #include "Text3DInternalTypes.h"
 #include "Utilities/Text3DUtilities.h"
@@ -31,10 +30,7 @@ void UText3DStaticMeshesRenderer::OnCreate()
 		TextRoot->SetupAttachment(TextComponent);
 	}
 
-	for (const TObjectPtr<UStaticMeshComponent>& CharacterMesh : CharacterMeshes)
-	{
-		CharacterMesh->RegisterComponent();
-	}
+	AllocateComponents(0);
 }
 
 void UText3DStaticMeshesRenderer::OnUpdate(const UE::Text3D::Renderer::FUpdateParameters& InParameters)
@@ -210,7 +206,8 @@ void UText3DStaticMeshesRenderer::AllocateComponents(int32 InCount)
 
 	const int32 RemainingCharacterCount = CharacterMeshes.Num() - InCount;
 
-	AActor* Actor = GetTypedOuter<AActor>();
+	const UText3DComponent* TextComponent = GetText3DComponent();
+	AActor* Actor = TextComponent->GetOwner();
 
 	if (RemainingCharacterCount > 0)
 	{
@@ -222,7 +219,10 @@ void UText3DStaticMeshesRenderer::AllocateComponents(int32 InCount)
 			if (UStaticMeshComponent* MeshComponent = CharacterMeshes[CharacterIndex])
 			{
 				MeshComponent->DetachFromComponent(DetachmentRule);
-				MeshComponent->UnregisterComponent();
+				if (MeshComponent->IsRegistered())
+				{
+					MeshComponent->UnregisterComponent();
+				}
 				Actor->RemoveInstanceComponent(MeshComponent);
 				Actor->RemoveOwnedComponent(MeshComponent);
 				CharacterMeshesPool.Add(MeshComponent);
@@ -256,24 +256,6 @@ void UText3DStaticMeshesRenderer::AllocateComponents(int32 InCount)
 	CharacterMeshes.SetNum(InCount);
 }
 
-void UText3DStaticMeshesRenderer::CollectComponents()
-{
-	TArray<UObject*> OwnedObjects;
-	GetObjectsWithOuter(this, OwnedObjects, /** Include nested */false);
-
-	for (UObject* OwnedObject : OwnedObjects)
-	{
-		if (UStaticMeshComponent* MeshComponent = Cast<UStaticMeshComponent>(OwnedObject))
-		{
-			if (!CharacterMeshes.Contains(MeshComponent)
-				&& !CharacterMeshesPool.Contains(MeshComponent))
-			{
-				CharacterMeshes.Add(MeshComponent);
-			}
-		}
-	}
-}
-
 int32 UText3DStaticMeshesRenderer::GetGlyphCount()
 {
 	if (!TextRoot)
@@ -299,12 +281,6 @@ const TArray<UStaticMeshComponent*>& UText3DStaticMeshesRenderer::GetGlyphMeshCo
 	return CharacterMeshes;
 }
 
-void UText3DStaticMeshesRenderer::PostEditImport()
-{
-	Super::PostEditImport();
-	CollectComponents();
-}
-
 #if WITH_EDITOR
 void UText3DStaticMeshesRenderer::ConvertToStaticMesh()
 {
diff --git a/Engine/Plugins/VirtualProduction/Text3D/Source/Text3D/Private/Renderers/StaticMeshes/Text3DStaticMeshesRenderer.h b/Engine/Plugins/VirtualProduction/Text3D/Source/Text3D/Private/Renderers/StaticMeshes/Text3DStaticMeshesRenderer.h
--- a/Engine/Plugins/VirtualProduction/Text3D/Source/Text3D/Private/Renderers/StaticMeshes/Text3DStaticMeshesRenderer.h
+++ b/Engine/Plugins/VirtualProduction/Text3D/Source/Text3D/Private/Renderers/StaticMeshes/Text3DStaticMeshesRenderer.h
@@ -45,10 +45,6 @@ public:
 #endif
 
 protected:
-	//~ Begin UObject
-	virtual void PostEditImport() override;
-	//~ End UObject
-
 	//~ Begin UText3DRendererBase
 	virtual void OnCreate() override;
 	virtual void OnUpdate(const UE::Text3D::Renderer::FUpdateParameters& InParameters) override;
@@ -63,7 +59,6 @@ protected:
 	//~ End UText3DRendererBase
 
 	void AllocateComponents(int32 InCount);
-	void CollectComponents();
 
 #if WITH_EDITOR
 	void SetDebugMode(bool bEnabled);
diff --git a/Engine/Plugins/VirtualProduction/Text3D/Source/Text3D/Private/Renderers/Text3DRendererBase.cpp b/Engine/Plugins/VirtualProduction/Text3D/Source/Text3D/Private/Renderers/Text3DRendererBase.cpp
--- a/Engine/Plugins/VirtualProduction/Text3D/Source/Text3D/Private/Renderers/Text3DRendererBase.cpp
+++ b/Engine/Plugins/VirtualProduction/Text3D/Source/Text3D/Private/Renderers/Text3DRendererBase.cpp
@@ -107,7 +107,7 @@ void UText3DRendererBase::Update(EText3DRendererFlags InFlags)
 		return;
 	}
 
-	TArray<UText3DExtensionBase*, TInlineAllocator<6>> Extensions
+	TArray<UText3DExtensionBase*, TInlineAllocator<7>> Extensions
 	{
 		Text3DComponent->GetTokenExtension(),
 		Text3DComponent->GetStyleExtension(),
diff --git a/Engine/Plugins/VirtualProduction/Text3D/Source/Text3D/Private/Text3DComponent.cpp b/Engine/Plugins/VirtualProduction/Text3D/Source/Text3D/Private/Text3DComponent.cpp
--- a/Engine/Plugins/VirtualProduction/Text3D/Source/Text3D/Private/Text3DComponent.cpp
+++ b/Engine/Plugins/VirtualProduction/Text3D/Source/Text3D/Private/Text3DComponent.cpp
@@ -66,8 +66,7 @@ UText3DComponent::UText3DComponent()
 void UText3DComponent::OnComponentCreated()
 {
 	Super::OnComponentCreated();
-
-	OnTextRendererClassChanged();
+	RequestUpdate(EText3DRendererFlags::All, /** Immediate */ true);
 }
 
 void UText3DComponent::OnComponentDestroyed(bool bInDestroyingHierarchy)
@@ -85,19 +84,21 @@ void UText3DComponent::OnComponentDestroyed(bool bInDestroyingHierarchy)
 	}
 
 	TransformUpdated.RemoveAll(this);
+
+	FTSTicker::GetCoreTicker().RemoveTicker(TextUpdateHandle);
+	TextUpdateHandle.Reset();
 }
 
 void UText3DComponent::PostLoad()
 {
 	Super::PostLoad();
-
 	RefreshTypeface();
-	OnTextRendererClassChanged();
+	RequestUpdate(EText3DRendererFlags::All, /** Immediate */ true);
 }
 
 void UText3DComponent::RequestUpdate(EText3DRendererFlags InFlags, bool bInImmediate)
 {
-	if (InFlags != EText3DRendererFlags::None && !IsTemplate(RF_ClassDefaultObject))
+	if (InFlags != EText3DRendererFlags::None && !IsTemplate())
 	{
 		if (bInImmediate)
 		{
@@ -258,10 +259,9 @@ void UText3DComponent::OnRegister()
 {
 	Super::OnRegister();
 
-	if (TextRenderer && !TextRenderer->IsInitialized())
+	if (TextRenderer)
 	{
-		constexpr bool bImmediate = true;
-		RequestUpdate(EText3DRendererFlags::All, bImmediate);
+		RequestUpdate(UpdateFlags, /** Immediate */true);
 	}
 }
 
@@ -278,6 +278,12 @@ void UText3DComponent::OnUnregister()
 	}
 }
 
+void UText3DComponent::PostApplyToComponent()
+{
+	Super::PostApplyToComponent();
+	RequestUpdate(EText3DRendererFlags::All);
+}
+
 void UText3DComponent::Serialize(FArchive& InArchive)
 {
 	InArchive.UsingCustomVersion(FText3DComponentVersion::GUID);
@@ -985,6 +991,12 @@ void UText3DComponent::QueueRebuild(EText3DRendererFlags InFlags)
 					// New rebuild requests should only happen when, while building text, existing 'built data' becomes outdated.
 					TextUpdateHandle.Reset();
 
+					if (!IsRegistered())
+					{
+						// Handled in OnRegister()
+						return false;
+					}
+
 					const EText3DRendererFlags CurrentUpdateFlags = UpdateFlags;
 					UpdateFlags = EText3DRendererFlags::None;
 
@@ -1016,12 +1028,6 @@ void UText3DComponent::Rebuild(EText3DRendererFlags InUpdateFlags)
 		}
 	}
 
-	if (!IsRegistered())
-	{
-		UE_LOG(LogText3D, Warning, TEXT("Text3DComponent is not registered yet, cannot proceed"));
-		return;
-	}
-
 	if (!Font)
 	{
 		if (const UText3DProjectSettings* Text3DSettings = UText3DProjectSettings::Get())
@@ -1036,14 +1042,25 @@ void UText3DComponent::Rebuild(EText3DRendererFlags InUpdateFlags)
 		}
 	}
 
+	if (!IsRegistered())
+	{
+		if (!EnumHasAllFlags(UpdateFlags, InUpdateFlags))
+		{
+			UE_LOG(LogText3D, Log, TEXT("[%s] Rebuild called with update flags (%u) while component is not registered yet. This new Rebuild will be deferred to the next tick."), *GetReadableName(), InUpdateFlags);
+			QueueRebuild(InUpdateFlags);
+		}
+
+		return;
+	}
+
 	if (UE::IsSavingPackage(this))
 	{
 		if (!EnumHasAllFlags(UpdateFlags, InUpdateFlags))
 		{
 			UE_LOG(LogText3D, Log, TEXT("[%s] Rebuild called with update flags (%u) while package is being saved. This new Rebuild will be deferred to the next tick."), *GetReadableName(), InUpdateFlags);
+			QueueRebuild(InUpdateFlags);
 		}
 
-		QueueRebuild(InUpdateFlags);
 		return;
 	}
 
@@ -1052,9 +1069,9 @@ void UText3DComponent::Rebuild(EText3DRendererFlags InUpdateFlags)
 		if (!EnumHasAllFlags(UpdateFlags, InUpdateFlags))
 		{
 			UE_LOG(LogText3D, Log, TEXT("[%s] Rebuild called with update flags (%u) while building text. This new Rebuild will be deferred to the next tick."), *GetReadableName(), InUpdateFlags);
+			QueueRebuild(InUpdateFlags);
 		}
 
-		QueueRebuild(InUpdateFlags);
 		return;
 	}
 
@@ -1141,7 +1158,7 @@ void UText3DComponent::OnTextRendererClassChanged()
 	const AActor* ContextActor = GetOwner();
 	UE_LOG(LogText3D, Log, TEXT("%s : Text3D renderer changed : old %s - new %s"), !!ContextActor ? *ContextActor->GetActorNameOrLabel() : TEXT("Invalid owner"), *OldRendererName.ToString(), *TextRenderer->GetFriendlyName().ToString());
 
-	if (IsRegistered())
+	if (!bIsUpdatingText)
 	{
 		constexpr bool bImmediate = true;
 		RequestUpdate(EText3DRendererFlags::All, bImmediate);
@@ -1157,7 +1174,6 @@ void UText3DComponent::OnTextChanged()
 void UText3DComponent::OnFontPropertiesChanged()
 {
 	RefreshTypeface();
-
 	RequestUpdate(EText3DRendererFlags::All);
 }
 
@@ -1170,22 +1186,14 @@ void UText3DComponent::CleanupRenderer()
 	{
 		if (UText3DRendererBase* Renderer = Cast<UText3DRendererBase>(OwnedObject))
 		{
-			if (!TextRenderer && Renderer->GetClass() == TextRendererClass)
-			{
-				TextRenderer = Renderer;
-			}
-			else if (Renderer != TextRenderer)
+			if (Renderer != TextRenderer)
 			{
+				Renderer->Clear();
 				Renderer->Destroy();
 				Renderer->MarkAsGarbage();
 			}
 		}
 	}
-
-	if (TextRenderer)
-	{
-		TextRenderer->Clear();
-	}
 }
 
 TArray<FName> UText3DComponent::GetTypefaceNames() const
diff --git a/Engine/Plugins/VirtualProduction/Text3D/Source/Text3D/Public/Characters/Text3DCharacterBase.h b/Engine/Plugins/VirtualProduction/Text3D/Source/Text3D/Public/Characters/Text3DCharacterBase.h
--- a/Engine/Plugins/VirtualProduction/Text3D/Source/Text3D/Public/Characters/Text3DCharacterBase.h
+++ b/Engine/Plugins/VirtualProduction/Text3D/Source/Text3D/Public/Characters/Text3DCharacterBase.h
@@ -110,12 +110,13 @@ protected:
 	UPROPERTY(EditInstanceOnly, Setter="SetVisibility", Getter="GetVisibility", Category="Character", meta=(AllowPrivateAccess = "true"))
 	bool bVisible = true;
 
+	/** Tag used to identify the custom style applied to this character */
+	UPROPERTY(VisibleAnywhere, Setter, Getter, Category="Character")
+	FName StyleTag = NAME_None;
+
 	/** Final transform after all extensions are applied */
 	FTransform Transform = FTransform::Identity;
 
 	/** Handle to the current glyph mesh this character represents */
 	UE::Text3D::Geometry::FCachedFontFaceGlyphHandle FontFaceGlyphHandle;
-
-	/** Tag used to identify the custom style applied to this character */
-	FName StyleTag = NAME_None;
 };
diff --git a/Engine/Plugins/VirtualProduction/Text3D/Source/Text3D/Public/Extensions/Text3DDefaultMaterialExtension.h b/Engine/Plugins/VirtualProduction/Text3D/Source/Text3D/Public/Extensions/Text3DDefaultMaterialExtension.h
--- a/Engine/Plugins/VirtualProduction/Text3D/Source/Text3D/Public/Extensions/Text3DDefaultMaterialExtension.h
+++ b/Engine/Plugins/VirtualProduction/Text3D/Source/Text3D/Public/Extensions/Text3DDefaultMaterialExtension.h
@@ -238,6 +238,7 @@ public:
 protected:
 	//~ Begin UObject
 	virtual void PostLoad() override;
+	virtual void PostInitProperties() override;
 #if WITH_EDITOR
 	virtual void PostEditChangeProperty(FPropertyChangedEvent& InEvent) override;
 	virtual void PostEditUndo() override;
diff --git a/Engine/Plugins/VirtualProduction/Text3D/Source/Text3D/Public/Extensions/Text3DDefaultStyleExtension.h b/Engine/Plugins/VirtualProduction/Text3D/Source/Text3D/Public/Extensions/Text3DDefaultStyleExtension.h
--- a/Engine/Plugins/VirtualProduction/Text3D/Source/Text3D/Public/Extensions/Text3DDefaultStyleExtension.h
+++ b/Engine/Plugins/VirtualProduction/Text3D/Source/Text3D/Public/Extensions/Text3DDefaultStyleExtension.h
@@ -58,7 +58,7 @@ protected:
 	 * Define your styles and use them in text like this <MyStyle>...</>,
 	 * the text between start and end tag will have this style applied
 	 */
-	UPROPERTY(EditInstanceOnly, Instanced, Category="Style", NoClear, meta=(AllowPrivateAccess="true"))
+	UPROPERTY(EditAnywhere, Instanced, DisplayName="Styles Override", Category="Style", NoClear, meta=(AllowPrivateAccess="true"))
 	TArray<TObjectPtr<UText3DStyleBase>> Styles;
 
 private:
diff --git a/Engine/Plugins/VirtualProduction/Text3D/Source/Text3D/Public/Extensions/Text3DDefaultTokenExtension.h b/Engine/Plugins/VirtualProduction/Text3D/Source/Text3D/Public/Extensions/Text3DDefaultTokenExtension.h
--- a/Engine/Plugins/VirtualProduction/Text3D/Source/Text3D/Public/Extensions/Text3DDefaultTokenExtension.h
+++ b/Engine/Plugins/VirtualProduction/Text3D/Source/Text3D/Public/Extensions/Text3DDefaultTokenExtension.h
@@ -32,9 +32,10 @@ protected:
 
 	/**
 	 * Define your tokens and use them in text like this {MyCustomToken},
-	 * {MyCustomToken} will be swapped by the content of the token called MyCustomToken
+	 * {MyCustomToken} will be swapped by the content of the token called MyCustomToken,
+	 * This can be used with remote control or sequencer to swap part of the text
 	 */
-	UPROPERTY(EditInstanceOnly, Instanced, Category="Token", NoClear, meta=(AllowPrivateAccess="true"))
+	UPROPERTY(EditAnywhere, Instanced, Category="Token", NoClear, meta=(AllowPrivateAccess="true"))
 	TArray<TObjectPtr<UText3DTokenBase>> Tokens;
 
 	UPROPERTY(Transient)
diff --git a/Engine/Plugins/VirtualProduction/Text3D/Source/Text3D/Public/Text3DComponent.h b/Engine/Plugins/VirtualProduction/Text3D/Source/Text3D/Public/Text3DComponent.h
--- a/Engine/Plugins/VirtualProduction/Text3D/Source/Text3D/Public/Text3DComponent.h
+++ b/Engine/Plugins/VirtualProduction/Text3D/Source/Text3D/Public/Text3DComponent.h
@@ -7,8 +7,6 @@
 #include "Text3DTypes.h"
 #include "Text3DComponent.generated.h"
 
-class FTextLayout;
-class ITextLayoutMarshaller;
 class UFont;
 class UMaterialInterface;
 class UText3DCharacterBase;
@@ -63,14 +61,15 @@ public:
 		return TextPreUpdateDelegate;
 	}
 
-	/** Get the text value and signal the primitives to be rebuilt */
+	/** Get the text value */
+	UFUNCTION(BlueprintPure, Category = "Text3D")
 	const FText& GetText() const
 	{
 		return Text;
 	}
 
 	/** Set the text value and signal the primitives to be rebuilt */
-	UFUNCTION(BlueprintCallable, Category = "Rendering|Components|Text3D")
+	UFUNCTION(BlueprintCallable, Category = "Text3D")
 	TEXT3D_API void SetText(const FText& Value);
 
 	/**
@@ -80,169 +79,171 @@ public:
 	UFUNCTION(BlueprintCallable, Category = "Text3D")
 	TEXT3D_API const FText& GetFormattedText() const;
 
-	UFUNCTION(BlueprintCallable, Category = "Rendering|Components|Text3D")
+	UFUNCTION(BlueprintCallable, Category = "Text3D")
 	TEXT3D_API void SetEnforceUpperCase(bool bInEnforceUpperCase);
 
+	UFUNCTION(BlueprintPure, Category = "Text3D")
 	bool GetEnforceUpperCase() const
 	{
 		return bEnforceUpperCase;
 	}
 
-	/** Get the text font and signal the primitives to be rebuilt */
+	/** Get the text font */
+	UFUNCTION(BlueprintPure, Category = "Text3D")
 	UFont* GetFont() const
 	{
 		return Font;
 	}
 
 	/** Set the text font and signal the primitives to be rebuilt */
-	UFUNCTION(BlueprintCallable, Category = "Rendering|Components|Text3D")
+	UFUNCTION(BlueprintCallable, Category = "Text3D")
 	TEXT3D_API void SetFont(UFont* InFont);
 
 	/** Set the text font size */
-	UFUNCTION(BlueprintCallable, Category = "Rendering|Components|Text3D")
+	UFUNCTION(BlueprintCallable, Category = "Text3D")
 	TEXT3D_API void SetFontSize(float InSize);
 
-	UFUNCTION(BlueprintPure, Category = "Rendering|Components|Text3D")
+	UFUNCTION(BlueprintPure, Category = "Text3D")
 	TEXT3D_API float GetFontSize() const;
 
 	/** Set whether an outline is applied. */
-	UFUNCTION(BlueprintCallable, Category = "Rendering|Components|Text3D")
+	UFUNCTION(BlueprintCallable, Category = "Text3D")
 	TEXT3D_API void SetHasOutline(const bool bValue);
 
 	/** Set the outline width. */
-	UFUNCTION(BlueprintCallable, Category = "Rendering|Components|Text3D")
+	UFUNCTION(BlueprintCallable, Category = "Text3D")
 	TEXT3D_API void SetOutlineExpand(const float Value);
 
 	/** Set the text extrusion size and signal the primitives to be rebuilt */
-	UFUNCTION(BlueprintCallable, Category = "Rendering|Components|Text3D")
+	UFUNCTION(BlueprintCallable, Category = "Text3D")
 	TEXT3D_API void SetExtrude(const float Value);
 
 	/** Set the 3d bevel value */
-	UFUNCTION(BlueprintCallable, Category = "Rendering|Components|Text3D")
+	UFUNCTION(BlueprintCallable, Category = "Text3D")
 	TEXT3D_API void SetBevel(const float Value);
 
 	/** Set the 3d bevel type */
-	UFUNCTION(BlueprintCallable, Category = "Rendering|Components|Text3D")
+	UFUNCTION(BlueprintCallable, Category = "Text3D")
 	TEXT3D_API void SetBevelType(const EText3DBevelType Value);
 
 	/** Set the amount of segments that will be used to tessellate the Bevel */
-	UFUNCTION(BlueprintCallable, Category = "Rendering|Components|Text3D")
+	UFUNCTION(BlueprintCallable, Category = "Text3D")
 	TEXT3D_API void SetBevelSegments(const int32 Value);
 
 	/** Get the text front material */
 	TEXT3D_API UMaterialInterface* GetFrontMaterial() const;
 
 	/** Set the text front material */
-	UFUNCTION(BlueprintCallable, Category = "Rendering|Components|Text3D")
+	UFUNCTION(BlueprintCallable, Category = "Text3D")
 	TEXT3D_API void SetFrontMaterial(UMaterialInterface* Value);
 
 	/** Get the text bevel material */
 	TEXT3D_API UMaterialInterface* GetBevelMaterial() const;
 
 	/** Set the text bevel material */
-	UFUNCTION(BlueprintCallable, Category = "Rendering|Components|Text3D")
+	UFUNCTION(BlueprintCallable, Category = "Text3D")
 	TEXT3D_API void SetBevelMaterial(UMaterialInterface* Value);
 	
 	/** Get the text extrude material */
 	TEXT3D_API UMaterialInterface* GetExtrudeMaterial() const;
 
 	/** Set the text extrude material */
-	UFUNCTION(BlueprintCallable, Category = "Rendering|Components|Text3D")
+	UFUNCTION(BlueprintCallable, Category = "Text3D")
 	TEXT3D_API void SetExtrudeMaterial(UMaterialInterface* Value);
 
 	/** Get the text back material */
 	TEXT3D_API UMaterialInterface* GetBackMaterial() const;
 
 	/** Set the text back material */
-	UFUNCTION(BlueprintCallable, Category = "Rendering|Components|Text3D")
+	UFUNCTION(BlueprintCallable, Category = "Text3D")
 	TEXT3D_API void SetBackMaterial(UMaterialInterface* Value);
 
 	/** Get the kerning value and signal the primitives to be rebuilt */
 	TEXT3D_API float GetKerning() const;
 
 	/** Set the kerning value and signal the primitives to be rebuilt */
-	UFUNCTION(BlueprintCallable, Category = "Rendering|Components|Text3D")
+	UFUNCTION(BlueprintCallable, Category = "Text3D")
 	TEXT3D_API void SetKerning(const float Value);
 
 	/** Get the line spacing value and signal the primitives to be rebuilt */
 	TEXT3D_API float GetLineSpacing() const;
 
 	/** Set the line spacing value and signal the primitives to be rebuilt */
-	UFUNCTION(BlueprintCallable, Category = "Rendering|Components|Text3D")
+	UFUNCTION(BlueprintCallable, Category = "Text3D")
 	TEXT3D_API void SetLineSpacing(const float Value);
 
 	/** Get the word spacing value and signal the primitives to be rebuilt */
 	TEXT3D_API float GetWordSpacing() const;
 
 	/** Set the word spacing value and signal the primitives to be rebuilt */
-	UFUNCTION(BlueprintCallable, Category = "Rendering|Components|Text3D")
+	UFUNCTION(BlueprintCallable, Category = "Text3D")
 	TEXT3D_API void SetWordSpacing(const float Value);
 
 	/** Get the horizontal alignment value and signal the primitives to be rebuilt */
 	TEXT3D_API EText3DHorizontalTextAlignment GetHorizontalAlignment() const;
 
 	/** Set the horizontal alignment value and signal the primitives to be rebuilt */
-	UFUNCTION(BlueprintCallable, Category = "Rendering|Components|Text3D")
+	UFUNCTION(BlueprintCallable, Category = "Text3D")
 	TEXT3D_API void SetHorizontalAlignment(const EText3DHorizontalTextAlignment value);
 
 	/** Get the vertical alignment and signal the primitives to be rebuilt */
 	TEXT3D_API EText3DVerticalTextAlignment GetVerticalAlignment() const;
 
 	/** Set the vertical alignment and signal the primitives to be rebuilt */
-	UFUNCTION(BlueprintCallable, Category = "Rendering|Components|Text3D")
+	UFUNCTION(BlueprintCallable, Category = "Text3D")
 	TEXT3D_API void SetVerticalAlignment(const EText3DVerticalTextAlignment value);
 
 	/** Whether a maximum width is specified */
 	TEXT3D_API bool HasMaxWidth() const;
 
 	/** Enable / Disable a Maximum Width */
-	UFUNCTION(BlueprintCallable, Category = "Rendering|Components|Text3D")
+	UFUNCTION(BlueprintCallable, Category = "Text3D")
 	TEXT3D_API void SetHasMaxWidth(const bool Value);
 
 	/** Get the Maximum Width - If width is larger, mesh will scale down to fit MaxWidth value */
 	TEXT3D_API float GetMaxWidth() const;
 
 	/** Set the Maximum Width - If width is larger, mesh will scale down to fit MaxWidth value */
-	UFUNCTION(BlueprintCallable, Category = "Rendering|Components|Text3D")
+	UFUNCTION(BlueprintCallable, Category = "Text3D")
 	TEXT3D_API void SetMaxWidth(const float Value);
 
 	/** Get the Maximum Width Handling - Whether to wrap before scaling when the text size reaches the max width */
 	TEXT3D_API EText3DMaxWidthHandling GetMaxWidthHandling() const;
 
 	/** Set the Maximum Width Handling - Whether to wrap before scaling when the text size reaches the max width */
-	UFUNCTION(BlueprintCallable, Category = "Rendering|Components|Text3D")
+	UFUNCTION(BlueprintCallable, Category = "Text3D")
 	TEXT3D_API void SetMaxWidthHandling(const EText3DMaxWidthHandling Value);
 
 	/** Whether a maximum height is specified */
 	TEXT3D_API bool HasMaxHeight() const;
 
 	/** Enable / Disable a Maximum Height */
-	UFUNCTION(BlueprintCallable, Category = "Rendering|Components|Text3D")
+	UFUNCTION(BlueprintCallable, Category = "Text3D")
 	TEXT3D_API void SetHasMaxHeight(const bool Value);
 
 	/** Get the Maximum Height - If height is larger, mesh will scale down to fit MaxHeight value */
 	TEXT3D_API float GetMaxHeight() const;
 
 	/** Set the Maximum Height - If height is larger, mesh will scale down to fit MaxHeight value */
-	UFUNCTION(BlueprintCallable, Category = "Rendering|Components|Text3D")
+	UFUNCTION(BlueprintCallable, Category = "Text3D")
 	TEXT3D_API void SetMaxHeight(const float Value);
 
 	/** Get if the mesh should scale proportionally when Max Width/Height is set */
 	TEXT3D_API bool ScalesProportionally() const;
 
 	/** Set if the mesh should scale proportionally when Max Width/Height is set */
-	UFUNCTION(BlueprintCallable, Category = "Rendering|Components|Text3D")
+	UFUNCTION(BlueprintCallable, Category = "Text3D")
 	TEXT3D_API void SetScaleProportionally(const bool Value);
 
 	/** Get the value of CastShadow. */
 	TEXT3D_API bool CastsShadow() const;
 
 	/** Set the value of CastShadow. */
-	UFUNCTION(BlueprintCallable, Category = "Rendering|Components|Text3D")
+	UFUNCTION(BlueprintCallable, Category = "Text3D")
 	TEXT3D_API void SetCastShadow(bool NewCastShadow);
 
 	/** Get whole text rendered bounds */
-	UFUNCTION(BlueprintCallable, Category = "Rendering|Components|Text3D")
+	UFUNCTION(BlueprintCallable, Category = "Text3D")
 	TEXT3D_API void GetBounds(FVector& Origin, FVector& BoxExtent) const;
 
 	/** Get whole text rendered bounds */
@@ -416,6 +417,7 @@ protected:
 	TEXT3D_API virtual void OnComponentDestroyed(bool bInDestroyingHierarchy) override;
 	TEXT3D_API virtual void OnRegister() override;
 	TEXT3D_API virtual void OnUnregister() override;
+	TEXT3D_API virtual void PostApplyToComponent() override;
 	//~ End UActorComponent
 	
 	//~ Begin UObject
