MotionDesign : Text3D
- Fixed duplicated Text3D components inside an actor fails to be displayed/rendered
- Marking Text3D as beta for at least one version due to recent changes and features introduced

diff --git a/Engine/Plugins/VirtualProduction/Text3D/Source/Text3D/Private/Text3DComponent.cpp b/Engine/Plugins/VirtualProduction/Text3D/Source/Text3D/Private/Text3DComponent.cpp
--- a/Engine/Plugins/VirtualProduction/Text3D/Source/Text3D/Private/Text3DComponent.cpp
+++ b/Engine/Plugins/VirtualProduction/Text3D/Source/Text3D/Private/Text3DComponent.cpp
@@ -255,6 +255,22 @@ void UText3DComponent::OnComponentTransformed(USceneComponent*, EUpdateTransform
 	}
 }
 
+#if WITH_EDITOR
+void UText3DComponent::PostEditChangeOwner()
+{
+	// Refresh nested component tree in editor after user duplication to force subcomponents to register/attach
+	FTSTicker::GetCoreTicker().AddTicker(FTickerDelegate::CreateWeakLambda(this, [this](float)->bool
+	{
+		if (AActor* TextActor = GetOwner())
+		{
+			TextActor->PostEditChange();
+		}
+
+		return false;
+	}));
+}
+#endif
+
 void UText3DComponent::OnRegister()
 {
 	Super::OnRegister();
@@ -352,6 +368,15 @@ void UText3DComponent::PostEditImport()
 	RequestUpdate(EText3DRendererFlags::All);
 }
 
+void UText3DComponent::PostDuplicate(bool bInDuplicateForPIE)
+{
+	Super::PostDuplicate(bInDuplicateForPIE);
+
+#if WITH_EDITOR
+	PostEditChangeOwner();
+#endif
+}
+
 #if WITH_EDITOR
 void UText3DComponent::PostEditChangeProperty(FPropertyChangedEvent& PropertyChangedEvent)
 {
@@ -1162,6 +1187,10 @@ void UText3DComponent::OnTextRendererClassChanged()
 	{
 		constexpr bool bImmediate = true;
 		RequestUpdate(EText3DRendererFlags::All, bImmediate);
+
+#if WITH_EDITOR
+		PostEditChangeOwner();
+#endif
 	}
 }
 
diff --git a/Engine/Plugins/VirtualProduction/Text3D/Source/Text3D/Public/Text3DComponent.h b/Engine/Plugins/VirtualProduction/Text3D/Source/Text3D/Public/Text3DComponent.h
--- a/Engine/Plugins/VirtualProduction/Text3D/Source/Text3D/Public/Text3DComponent.h
+++ b/Engine/Plugins/VirtualProduction/Text3D/Source/Text3D/Public/Text3DComponent.h
@@ -424,6 +424,7 @@ protected:
 	TEXT3D_API virtual void Serialize(FArchive& InArchive) override;
 	TEXT3D_API virtual void PostLoad() override;
 	TEXT3D_API virtual void PostEditImport() override;
+	TEXT3D_API virtual void PostDuplicate(bool bInDuplicateForPIE) override;
 #if WITH_EDITOR
 	TEXT3D_API virtual void PostEditChangeProperty(FPropertyChangedEvent& PropertyChangedEvent) override;
 	TEXT3D_API virtual void PostTransacted(const FTransactionObjectEvent& TransactionEvent) override;
@@ -561,6 +562,10 @@ private:
 
 	void OnComponentTransformed(USceneComponent*, EUpdateTransformFlags, ETeleportType);
 
+#if WITH_EDITOR
+	void PostEditChangeOwner();
+#endif
+
 	bool bZeroScaled = false;
 
 #if WITH_EDITORONLY_DATA
diff --git a/Engine/Plugins/VirtualProduction/Text3D/Text3D.uplugin b/Engine/Plugins/VirtualProduction/Text3D/Text3D.uplugin
--- a/Engine/Plugins/VirtualProduction/Text3D/Text3D.uplugin
+++ b/Engine/Plugins/VirtualProduction/Text3D/Text3D.uplugin
@@ -10,6 +10,7 @@
   "MarketplaceURL": "",
   "SupportURL": "",
   "CanContainContent": true,
+  "IsBetaVersion": true,
   "Installed": false,
   "Modules": [
     {
