MotionDesign : Text3D
- Fixed applying a lineargradient to text breaks (rotated or not)

diff --git a/Engine/Plugins/VirtualProduction/Text3D/Source/Text3D/Private/Extensions/Text3DDefaultMaterialExtension.cpp b/Engine/Plugins/VirtualProduction/Text3D/Source/Text3D/Private/Extensions/Text3DDefaultMaterialExtension.cpp
--- a/Engine/Plugins/VirtualProduction/Text3D/Source/Text3D/Private/Extensions/Text3DDefaultMaterialExtension.cpp
+++ b/Engine/Plugins/VirtualProduction/Text3D/Source/Text3D/Private/Extensions/Text3DDefaultMaterialExtension.cpp
@@ -337,23 +337,21 @@ EText3DExtensionResult UText3DDefaultMaterialExtension::PreRendererUpdate(const
 
 		if (BlendMode == EText3DMaterialBlendMode::Translucent || Style == EText3DMaterialStyle::Gradient)
 		{
-			const FBox LocalBounds = Text3DComponent->GetBounds();
-			const FVector LocalBoundsExtent = Text3DComponent->GetRelativeRotation().UnrotateVector(LocalBounds.GetSize());
-
-			const UText3DLayoutExtensionBase* LayoutExtension = Text3DComponent->GetLayoutExtension();
-
-			if (!LayoutExtension)
+			if (EnumHasAnyFlags(InParameters.UpdateFlags, EText3DRendererFlags::Geometry)
+				|| EnumHasAnyFlags(InParameters.UpdateFlags, EText3DRendererFlags::Layout))
 			{
-				UE_LOG(LogText3D, Warning, TEXT("Invalid layout extension retrieved in material extension, cannot proceed."))
-				return EText3DExtensionResult::Failed;
+				// Convert to local space
+				FBox TextBounds = Text3DComponent->GetBounds();
+				const FTransform ComponentTransform = Text3DComponent->GetComponentTransform();
+				TextBounds = TextBounds.TransformBy(ComponentTransform.Inverse());
+				// Apply scale ratio based on font size 64 => 1.f
+				FTransform RatioTransform;
+				RatioTransform.SetScale3D(FVector(64.f/Text3DComponent->GetFontSize()));
+				LocalTextBounds = TextBounds.TransformBy(RatioTransform);
 			}
 
-			const FVector LineLocation = LayoutExtension->GetLineLocation(0);
-			const FVector TextScaleFactor = LayoutExtension->GetTextScale() * Text3DComponent->GetComponentScale();
-			const FVector BoundsSizeScaled = LocalBoundsExtent / TextScaleFactor;
-
-			SetVectorParameter(Materials, UText3DProjectSettings::FMaterialParameters::BoundsOrigin, LineLocation);
-			SetVectorParameter(Materials, UText3DProjectSettings::FMaterialParameters::BoundsSize, BoundsSizeScaled);
+			SetVectorParameter(Materials, UText3DProjectSettings::FMaterialParameters::BoundsOrigin, LocalTextBounds.Min);
+			SetVectorParameter(Materials, UText3DProjectSettings::FMaterialParameters::BoundsSize, LocalTextBounds.GetSize());
 		}
 
 		if (BlendMode == EText3DMaterialBlendMode::Translucent)
@@ -545,8 +543,11 @@ UMaterialInstanceDynamic* UText3DDefaultMaterialExtension::FindOrAdd(EText3DGrou
 	}
 
 	UMaterialInstanceDynamic* DynamicMaterial = UMaterialInstanceDynamic::Create(ParentMaterial, this);
-	DynamicMaterial->SetFlags(RF_Transient | RF_Public);
-	GroupDynamicMaterials.Emplace(GroupKey, DynamicMaterial);
+	if (DynamicMaterial)
+	{
+		DynamicMaterial->SetFlags(RF_Transient | RF_Public);
+		GroupDynamicMaterials.Emplace(GroupKey, DynamicMaterial);
+	}
 
 	return DynamicMaterial;
 }
diff --git a/Engine/Plugins/VirtualProduction/Text3D/Source/Text3D/Public/Extensions/Text3DDefaultMaterialExtension.h b/Engine/Plugins/VirtualProduction/Text3D/Source/Text3D/Public/Extensions/Text3DDefaultMaterialExtension.h
--- a/Engine/Plugins/VirtualProduction/Text3D/Source/Text3D/Public/Extensions/Text3DDefaultMaterialExtension.h
+++ b/Engine/Plugins/VirtualProduction/Text3D/Source/Text3D/Public/Extensions/Text3DDefaultMaterialExtension.h
@@ -354,4 +354,7 @@ protected:
 	/** Cached group dynamic materials created during session based on their options */
 	UPROPERTY(Transient, DuplicateTransient)
 	TMap<FText3DMaterialGroupKey, TObjectPtr<UMaterialInstanceDynamic>> GroupDynamicMaterials;
+
+	/** Used to pass bounds to material */
+	FBox LocalTextBounds = FBox(ForceInitToZero);
 };
