MotionDesign : Text3D
- Fixed Arabic diacritics marks are not laid out properly

diff --git a/Engine/Plugins/VirtualProduction/Text3D/Source/Text3D/Private/Extensions/Text3DDefaultLayoutExtension.cpp b/Engine/Plugins/VirtualProduction/Text3D/Source/Text3D/Private/Extensions/Text3DDefaultLayoutExtension.cpp
--- a/Engine/Plugins/VirtualProduction/Text3D/Source/Text3D/Private/Extensions/Text3DDefaultLayoutExtension.cpp
+++ b/Engine/Plugins/VirtualProduction/Text3D/Source/Text3D/Private/Extensions/Text3DDefaultLayoutExtension.cpp
@@ -343,6 +343,13 @@ EText3DExtensionResult UText3DDefaultLayoutExtension::PreRendererUpdate(const UE
 					Shift.Y += Tracking + GlyphText->Kernings[CharacterIndex];
 				}
 
+				// Handle special characters with no advance but offset
+				if (GlyphEntry->Entry.XAdvance == 0)
+				{
+					Shift.Y += GlyphEntry->Entry.XOffset;
+					Shift.Z -= GlyphEntry->Entry.YOffset;
+				}
+
 				FTransform Transform(Location + Shift);
 				const FVector2D FaceSize = CachedMesh->FontFaceGlyphSize;
 
diff --git a/Engine/Plugins/VirtualProduction/Text3D/Source/Text3D/Private/LayoutBuilders/Text3DShapedGlyphLine.cpp b/Engine/Plugins/VirtualProduction/Text3D/Source/Text3D/Private/LayoutBuilders/Text3DShapedGlyphLine.cpp
--- a/Engine/Plugins/VirtualProduction/Text3D/Source/Text3D/Private/LayoutBuilders/Text3DShapedGlyphLine.cpp
+++ b/Engine/Plugins/VirtualProduction/Text3D/Source/Text3D/Private/LayoutBuilders/Text3DShapedGlyphLine.cpp
@@ -7,7 +7,7 @@ float UE::Text3D::Layout::FGlyphLine::GetWidthAdvance(const int32 InIndex, const
 	check(InIndex >= 0 && InIndex < Glyphs.Num());
 
 	const FGlyphEntry& Glyph = Glyphs[InIndex];
-	float Advance = Glyph.Entry.XOffset + Glyph.Entry.XAdvance;
+	float Advance = (Glyph.Entry.XAdvance != 0) ? Glyph.Entry.XOffset + Glyph.Entry.XAdvance : Glyph.Entry.XAdvance;
 
 	const bool bFirstCharacter = InIndex == 0;
 	const bool bLastCharacter = InIndex == Glyphs.Num() - 1;
@@ -16,7 +16,7 @@ float UE::Text3D::Layout::FGlyphLine::GetWidthAdvance(const int32 InIndex, const
 	if (!bSkipCharacter)
 	{
 		// @note: as per FSlateElementBatcher::BuildShapedTextSequence, per Glyph Kerning isn't used
-		Advance += InKerning;
+		Advance += (Glyph.Entry.XAdvance != 0) ? InKerning : 0.f;
 
 		if (!Glyph.Entry.bIsVisible)
 		{
