MotionDesign : Text3D
- Fixed server cooking failing due to error caused by FreeType missing
- Fixed text3D blueprint infinite load loop caused by a previous fix
- Fixed spawnable/archetype behavior causing components to not attach themselves
- Fixed 5.6 broke compatibility and removed Text BP getters
- Fixed adding a text 3D component to an actor blueprint causes editor to freeze

diff --git a/Engine/Plugins/VirtualProduction/Avalanche/Source/AvalancheText/Public/AvaText3DComponent.h b/Engine/Plugins/VirtualProduction/Avalanche/Source/AvalancheText/Public/AvaText3DComponent.h
--- a/Engine/Plugins/VirtualProduction/Avalanche/Source/AvalancheText/Public/AvaText3DComponent.h
+++ b/Engine/Plugins/VirtualProduction/Avalanche/Source/AvalancheText/Public/AvaText3DComponent.h
@@ -10,8 +10,7 @@
 UCLASS(MinimalAPI
 	, ClassGroup="Text3D"
 	, PrioritizeCategories="Text Character Layout Geometry Material Rendering Transform"
-	, HideCategories=(Replication,Collision,HLOD,Physics,Networking,Input,Actor,Cooking,LevelInstance,Streaming,DataLayers,WorldPartition)
-	, meta=(BlueprintSpawnableComponent))
+	, HideCategories=(Replication,Collision,HLOD,Physics,Networking,Input,Actor,Cooking,LevelInstance,Streaming,DataLayers,WorldPartition))
 class UAvaText3DComponent : public UText3DComponent
 {
 	friend class FAvaTextVisualizer;
diff --git a/Engine/Plugins/VirtualProduction/Text3D/Source/Text3D/Private/Extensions/Text3DDefaultGeometryExtension.cpp b/Engine/Plugins/VirtualProduction/Text3D/Source/Text3D/Private/Extensions/Text3DDefaultGeometryExtension.cpp
--- a/Engine/Plugins/VirtualProduction/Text3D/Source/Text3D/Private/Extensions/Text3DDefaultGeometryExtension.cpp
+++ b/Engine/Plugins/VirtualProduction/Text3D/Source/Text3D/Private/Extensions/Text3DDefaultGeometryExtension.cpp
@@ -155,11 +155,6 @@ void UText3DDefaultGeometryExtension::PostEditChangeProperty(FPropertyChangedEve
 }
 #endif
 
-const FText3DCachedMesh* UText3DDefaultGeometryExtension::FindOrLoadGlyphMesh(UText3DCharacterBase* InCharacter) const
-{
-	return InCharacter->GetGlyphMesh();
-}
-
 EText3DHorizontalTextAlignment UText3DDefaultGeometryExtension::GetGlyphHAlignment() const
 {
 	return PivotHAlignment;
diff --git a/Engine/Plugins/VirtualProduction/Text3D/Source/Text3D/Private/Renderers/DynamicMesh/Text3DDynamicMeshRenderer.cpp b/Engine/Plugins/VirtualProduction/Text3D/Source/Text3D/Private/Renderers/DynamicMesh/Text3DDynamicMeshRenderer.cpp
--- a/Engine/Plugins/VirtualProduction/Text3D/Source/Text3D/Private/Renderers/DynamicMesh/Text3DDynamicMeshRenderer.cpp
+++ b/Engine/Plugins/VirtualProduction/Text3D/Source/Text3D/Private/Renderers/DynamicMesh/Text3DDynamicMeshRenderer.cpp
@@ -52,9 +52,10 @@ void UText3DDynamicMeshRenderer::FScopedMeshEditor::ProcessMesh(const TFunctionR
 
 void UText3DDynamicMeshRenderer::OnCreate()
 {
+	UText3DComponent* TextComponent = GetText3DComponent();
+
 	if (!DynamicMeshComponent)
 	{
-		UText3DComponent* TextComponent = GetText3DComponent();
 		DynamicMeshComponent = NewObject<UDynamicMeshComponent>(this, TEXT("Text3DDynamicMesh"));
 		DynamicMeshComponent->SetupAttachment(TextComponent);
 
@@ -65,6 +66,14 @@ void UText3DDynamicMeshRenderer::OnCreate()
 
 		ClearMesh();
 	}
+	else
+	{
+		FDetachmentTransformRules DetachRule = FDetachmentTransformRules::KeepRelativeTransform;
+		DetachRule.bCallModify = false;
+		DynamicMeshComponent->DetachFromComponent(DetachRule);
+		const FAttachmentTransformRules AttachRule = FAttachmentTransformRules::KeepRelativeTransform;
+		DynamicMeshComponent->AttachToComponent(TextComponent, AttachRule);
+	}
 }
 
 void UText3DDynamicMeshRenderer::OnUpdate(const UE::Text3D::Renderer::FUpdateParameters& InParameters)
@@ -191,6 +200,20 @@ void UText3DDynamicMeshRenderer::OnDestroy()
 		DynamicMeshComponent->DestroyComponent();
 		DynamicMeshComponent = nullptr;
 	}
+
+	if (const UText3DComponent* TextComponent = GetText3DComponent())
+	{
+		TArray<USceneComponent*> ChildrenComponents;
+		TextComponent->GetChildrenComponents(/** IncludeChidren */true, ChildrenComponents);
+
+		for (USceneComponent* ChildComponent : ChildrenComponents)
+		{
+			if (ChildComponent->GetName() == TEXT("Text3DDynamicMesh"))
+			{
+				ChildComponent->DestroyComponent();
+			}
+		}
+	}
 }
 
 FName UText3DDynamicMeshRenderer::GetFriendlyName() const
diff --git a/Engine/Plugins/VirtualProduction/Text3D/Source/Text3D/Private/Renderers/StaticMeshes/Text3DStaticMeshesRenderer.cpp b/Engine/Plugins/VirtualProduction/Text3D/Source/Text3D/Private/Renderers/StaticMeshes/Text3DStaticMeshesRenderer.cpp
--- a/Engine/Plugins/VirtualProduction/Text3D/Source/Text3D/Private/Renderers/StaticMeshes/Text3DStaticMeshesRenderer.cpp
+++ b/Engine/Plugins/VirtualProduction/Text3D/Source/Text3D/Private/Renderers/StaticMeshes/Text3DStaticMeshesRenderer.cpp
@@ -23,12 +23,21 @@
 
 void UText3DStaticMeshesRenderer::OnCreate()
 {
+	UText3DComponent* TextComponent = GetText3DComponent();
+
 	if (!TextRoot)
 	{
-		UText3DComponent* TextComponent = GetText3DComponent();
 		TextRoot = NewObject<USceneComponent>(this, TEXT("TextRoot"));
 		TextRoot->SetupAttachment(TextComponent);
 	}
+	else
+	{
+		FDetachmentTransformRules DetachRule = FDetachmentTransformRules::KeepRelativeTransform;
+		DetachRule.bCallModify = false;
+		TextRoot->DetachFromComponent(DetachRule);
+		const FAttachmentTransformRules AttachRule = FAttachmentTransformRules::KeepRelativeTransform;
+		TextRoot->AttachToComponent(TextComponent, AttachRule);
+	}
 
 	AllocateComponents(0);
 }
@@ -175,6 +184,21 @@ void UText3DStaticMeshesRenderer::OnDestroy()
 		TextRoot->DestroyComponent();
 		TextRoot = nullptr;
 	}
+
+	if (const UText3DComponent* TextComponent = GetText3DComponent())
+	{
+		TArray<USceneComponent*> ChildrenComponents;
+		TextComponent->GetChildrenComponents(/** IncludeChidren */true, ChildrenComponents);
+
+		for (USceneComponent* ChildComponent : ChildrenComponents)
+		{
+			if (ChildComponent->GetName() == TEXT("TextRoot")
+				|| ChildComponent->GetName().StartsWith(TEXT("CharacterMesh")))
+			{
+				ChildComponent->DestroyComponent();
+			}
+		}
+	}
 }
 
 FName UText3DStaticMeshesRenderer::GetFriendlyName() const
diff --git a/Engine/Plugins/VirtualProduction/Text3D/Source/Text3D/Private/Renderers/Text3DRendererBase.cpp b/Engine/Plugins/VirtualProduction/Text3D/Source/Text3D/Private/Renderers/Text3DRendererBase.cpp
--- a/Engine/Plugins/VirtualProduction/Text3D/Source/Text3D/Private/Renderers/Text3DRendererBase.cpp
+++ b/Engine/Plugins/VirtualProduction/Text3D/Source/Text3D/Private/Renderers/Text3DRendererBase.cpp
@@ -249,4 +249,10 @@ FBox UText3DRendererBase::GetBounds() const
 	}
 
 	return CachedBounds.GetValue();
-}
\ No newline at end of file
+}
+
+void UText3DRendererBase::BeginDestroy()
+{
+	Super::BeginDestroy();
+	Destroy();
+}
diff --git a/Engine/Plugins/VirtualProduction/Text3D/Source/Text3D/Private/Text3DComponent.cpp b/Engine/Plugins/VirtualProduction/Text3D/Source/Text3D/Private/Text3DComponent.cpp
--- a/Engine/Plugins/VirtualProduction/Text3D/Source/Text3D/Private/Text3DComponent.cpp
+++ b/Engine/Plugins/VirtualProduction/Text3D/Source/Text3D/Private/Text3DComponent.cpp
@@ -93,7 +93,8 @@ void UText3DComponent::PostLoad()
 {
 	Super::PostLoad();
 	RefreshTypeface();
-	RequestUpdate(EText3DRendererFlags::All, /** Immediate */ true);
+	const bool bCompiledFromBlueprint = IsCompiledFromBlueprint();
+	RequestUpdate(EText3DRendererFlags::All, /** Immediate */!bCompiledFromBlueprint);
 }
 
 void UText3DComponent::RequestUpdate(EText3DRendererFlags InFlags, bool bInImmediate)
@@ -250,19 +251,34 @@ void UText3DComponent::OnComponentTransformed(USceneComponent*, EUpdateTransform
 	// Only update when scale was zero and is not anymore
 	if (bZeroScaled && !bIsZeroScaled)
 	{
-		RequestUpdate(EText3DRendererFlags::Layout);
+		RequestUpdate(EText3DRendererFlags::Layout | EText3DRendererFlags::Material);
 		bZeroScaled = false;
 	}
 }
 
-void UText3DComponent::OnRegister()
+bool UText3DComponent::IsCompiledFromBlueprint() const
 {
-	Super::OnRegister();
-
-	if (TextRenderer)
+	for (const UObject* Object = this; Object != nullptr && Object->GetOuter(); Object = Object->GetOuter())
 	{
-		RequestUpdate(UpdateFlags, /** Immediate */true);
+		if (Object->GetClass()->HasAllClassFlags(CLASS_CompiledFromBlueprint))
+		{
+			return true;
+		}
+
+		if (Object->IsA<AActor>())
+		{
+			break;
+		}
 	}
+
+	return false;
+}
+
+void UText3DComponent::OnRegister()
+{
+	Super::OnRegister();
+	const bool bCompiledFromBlueprint = IsCompiledFromBlueprint();
+	RequestUpdate(UpdateFlags, /** Immediate */!bCompiledFromBlueprint);
 }
 
 void UText3DComponent::OnUnregister()
@@ -352,7 +368,24 @@ void UText3DComponent::PostEditImport()
 	RequestUpdate(EText3DRendererFlags::All);
 }
 
+void UText3DComponent::PostDuplicate(bool bDuplicateForPIE)
+{
+	Super::PostDuplicate(bDuplicateForPIE);
+	RequestUpdate(EText3DRendererFlags::All);
+}
+
 #if WITH_EDITOR
+void UText3DComponent::PostCDOCompiled(const FPostCDOCompiledContext& InContext)
+{
+	Super::PostCDOCompiled(InContext);
+
+	// Cleanup old components
+	if (TextRenderer)
+	{
+		TextRenderer->Destroy();
+	}
+}
+
 void UText3DComponent::PostEditChangeProperty(FPropertyChangedEvent& PropertyChangedEvent)
 {
 	Super::PostEditChangeProperty(PropertyChangedEvent);
@@ -1017,6 +1050,11 @@ void UText3DComponent::Rebuild(EText3DRendererFlags InUpdateFlags)
 		return;
 	}
 
+#if !WITH_FREETYPE
+	UE_LOG(LogText3D, Warning, TEXT("FreeType is not available, cannot proceed without it"));
+	return;
+#endif
+
 	if (!TextRenderer || TextRenderer->GetClass() != TextRendererClass)
 	{
 		OnTextRendererClassChanged();
@@ -1114,20 +1152,27 @@ void UText3DComponent::Rebuild(EText3DRendererFlags InUpdateFlags)
 
 void UText3DComponent::OnTextRendererClassChanged()
 {
-	if (IsTemplate(RF_ClassDefaultObject))
+	if (IsTemplate())
 	{
 		return;
 	}
 
-	const UText3DRendererBase* RendererCDO = TextRendererClass.Get() ? TextRendererClass->GetDefaultObject<UText3DRendererBase>() : nullptr;
-
-	if (!RendererCDO)
+	if (!TextRendererClass)
 	{
 		if (TextRenderer)
 		{
 			TextRendererClass = TextRenderer->GetClass();
 		}
+		else
+		{
+			TextRendererClass = UText3DStaticMeshesRenderer::StaticClass();
+		}
+	}
+
+	const UText3DRendererBase* RendererCDO = TextRendererClass->GetDefaultObject<UText3DRendererBase>();
 
+	if (!RendererCDO)
+	{
 		return;
 	}
 
@@ -1160,8 +1205,7 @@ void UText3DComponent::OnTextRendererClassChanged()
 
 	if (!bIsUpdatingText)
 	{
-		constexpr bool bImmediate = true;
-		RequestUpdate(EText3DRendererFlags::All, bImmediate);
+		RequestUpdate(EText3DRendererFlags::All, /** Immediate */true);
 	}
 }
 
diff --git a/Engine/Plugins/VirtualProduction/Text3D/Source/Text3D/Public/Extensions/Text3DDefaultGeometryExtension.h b/Engine/Plugins/VirtualProduction/Text3D/Source/Text3D/Public/Extensions/Text3DDefaultGeometryExtension.h
--- a/Engine/Plugins/VirtualProduction/Text3D/Source/Text3D/Public/Extensions/Text3DDefaultGeometryExtension.h
+++ b/Engine/Plugins/VirtualProduction/Text3D/Source/Text3D/Public/Extensions/Text3DDefaultGeometryExtension.h
@@ -112,7 +112,6 @@ protected:
 	//~ End UObject
 
 	//~ Begin UText3DGeometryExtensionBase
-	virtual const FText3DCachedMesh* FindOrLoadGlyphMesh(UText3DCharacterBase* InCharacter) const override;
 	virtual EText3DHorizontalTextAlignment GetGlyphHAlignment() const override;
 	virtual EText3DVerticalTextAlignment GetGlyphVAlignment() const override;
 	virtual const FGlyphMeshParameters* GetGlyphMeshParameters() const override;
diff --git a/Engine/Plugins/VirtualProduction/Text3D/Source/Text3D/Public/Extensions/Text3DGeometryExtensionBase.h b/Engine/Plugins/VirtualProduction/Text3D/Source/Text3D/Public/Extensions/Text3DGeometryExtensionBase.h
--- a/Engine/Plugins/VirtualProduction/Text3D/Source/Text3D/Public/Extensions/Text3DGeometryExtensionBase.h
+++ b/Engine/Plugins/VirtualProduction/Text3D/Source/Text3D/Public/Extensions/Text3DGeometryExtensionBase.h
@@ -24,7 +24,8 @@ public:
 	{}
 
 	// Deprecated do not use anymore, use character getter directly
-	virtual const FText3DCachedMesh* FindOrLoadGlyphMesh(UText3DCharacterBase* InCharacter) const
+	UE_DEPRECATED(5.7, "Deprecated do not use anymore, use character getter directly")
+	virtual const FText3DCachedMesh* FindOrLoadGlyphMesh(UText3DCharacterBase* InCharacter) const final
 	{
 		return nullptr;
 	}
diff --git a/Engine/Plugins/VirtualProduction/Text3D/Source/Text3D/Public/Renderers/Text3DRendererBase.h b/Engine/Plugins/VirtualProduction/Text3D/Source/Text3D/Public/Renderers/Text3DRendererBase.h
--- a/Engine/Plugins/VirtualProduction/Text3D/Source/Text3D/Public/Renderers/Text3DRendererBase.h
+++ b/Engine/Plugins/VirtualProduction/Text3D/Source/Text3D/Public/Renderers/Text3DRendererBase.h
@@ -43,6 +43,10 @@ public:
 	}
 
 protected:
+	//~ Begin UObject
+	virtual void BeginDestroy() override;
+	//~ End UObject
+
 	/** Create and setup the implementation components, called on load or creation */
 	virtual void OnCreate() PURE_VIRTUAL(UText3DRenderingImplementationBase::OnCreate, );
 
diff --git a/Engine/Plugins/VirtualProduction/Text3D/Source/Text3D/Public/Text3DComponent.h b/Engine/Plugins/VirtualProduction/Text3D/Source/Text3D/Public/Text3DComponent.h
--- a/Engine/Plugins/VirtualProduction/Text3D/Source/Text3D/Public/Text3DComponent.h
+++ b/Engine/Plugins/VirtualProduction/Text3D/Source/Text3D/Public/Text3DComponent.h
@@ -62,14 +62,14 @@ public:
 	}
 
 	/** Get the text value */
-	UFUNCTION(BlueprintPure, Category = "Text3D")
+	UFUNCTION(BlueprintPure, BlueprintGetter, Category = "Text3D")
 	const FText& GetText() const
 	{
 		return Text;
 	}
 
 	/** Set the text value and signal the primitives to be rebuilt */
-	UFUNCTION(BlueprintCallable, Category = "Text3D")
+	UFUNCTION(BlueprintCallable, BlueprintSetter, Category = "Text3D")
 	TEXT3D_API void SetText(const FText& Value);
 
 	/**
@@ -79,24 +79,24 @@ public:
 	UFUNCTION(BlueprintCallable, Category = "Text3D")
 	TEXT3D_API const FText& GetFormattedText() const;
 
-	UFUNCTION(BlueprintCallable, Category = "Text3D")
+	UFUNCTION(BlueprintCallable, BlueprintSetter, Category = "Text3D")
 	TEXT3D_API void SetEnforceUpperCase(bool bInEnforceUpperCase);
 
-	UFUNCTION(BlueprintPure, Category = "Text3D")
+	UFUNCTION(BlueprintPure, BlueprintGetter, Category = "Text3D")
 	bool GetEnforceUpperCase() const
 	{
 		return bEnforceUpperCase;
 	}
 
 	/** Get the text font */
-	UFUNCTION(BlueprintPure, Category = "Text3D")
+	UFUNCTION(BlueprintPure, BlueprintGetter, Category = "Text3D")
 	UFont* GetFont() const
 	{
 		return Font;
 	}
 
 	/** Set the text font and signal the primitives to be rebuilt */
-	UFUNCTION(BlueprintCallable, Category = "Text3D")
+	UFUNCTION(BlueprintCallable, BlueprintSetter, Category = "Text3D")
 	TEXT3D_API void SetFont(UFont* InFont);
 
 	/** Set the text font size */
@@ -250,12 +250,14 @@ public:
 	FBox GetBounds() const;
 
 	/** Get the typeface */
+	UFUNCTION(BlueprintPure, BlueprintGetter, Category = "Text3D")
 	FName GetTypeface() const
 	{
 		return Typeface;
 	}
 
 	/** Set the typeface */
+	UFUNCTION(BlueprintCallable, BlueprintSetter, Category = "Text3D")
 	TEXT3D_API void SetTypeface(const FName InTypeface);
 
 	UFUNCTION(BlueprintCallable, Category = "Text3D")
@@ -424,7 +426,9 @@ protected:
 	TEXT3D_API virtual void Serialize(FArchive& InArchive) override;
 	TEXT3D_API virtual void PostLoad() override;
 	TEXT3D_API virtual void PostEditImport() override;
+	TEXT3D_API virtual void PostDuplicate(bool bDuplicateForPIE) override;
 #if WITH_EDITOR
+	TEXT3D_API virtual void PostCDOCompiled(const FPostCDOCompiledContext& InContext) override;
 	TEXT3D_API virtual void PostEditChangeProperty(FPropertyChangedEvent& PropertyChangedEvent) override;
 	TEXT3D_API virtual void PostTransacted(const FTransactionObjectEvent& TransactionEvent) override;
 #endif
@@ -476,15 +480,15 @@ protected:
 	TSubclassOf<UText3DRendererBase> TextRendererClass;
 
 	/** The text to generate a 3d mesh */
-	UPROPERTY(EditAnywhere, Getter, Setter, Category = "Text", meta = (MultiLine = true, AllowPrivateAccess = "true"))
+	UPROPERTY(EditAnywhere, Getter, Setter, BlueprintSetter="SetText", BlueprintGetter="GetText", Category = "Text", meta = (MultiLine = true, AllowPrivateAccess = "true"))
 	FText Text;
 
 	/** Text font defines the style of rendered characters */
-	UPROPERTY(EditAnywhere, Getter, Setter, DisplayName = "Default Font", Category = "Text", meta = (AdvancedFontPicker, AllowPrivateAccess = "true"))
+	UPROPERTY(EditAnywhere, Getter, Setter, BlueprintSetter="SetFont", BlueprintGetter="GetFont", DisplayName = "Default Font", Category = "Text", meta = (AdvancedFontPicker, AllowPrivateAccess = "true"))
 	TObjectPtr<UFont> Font;
 
 	/** Text font face, subset within font like bold, italic, regular */
-	UPROPERTY(EditAnywhere, Getter, Setter, DisplayName = "Default Typeface", Category = "Text", meta = (AllowPrivateAccess = "true", GetOptions="GetTypefaceNames"))
+	UPROPERTY(EditAnywhere, Getter, Setter, BlueprintSetter="SetTypeface", BlueprintGetter="GetTypeface", DisplayName = "Default Typeface", Category = "Text", meta = (AllowPrivateAccess = "true", GetOptions="GetTypefaceNames"))
 	FName Typeface;
 
 	/** Text font size defines the size of each rendered character */
@@ -492,7 +496,7 @@ protected:
 	float FontSize = 48.f;
 
 	/** Whether to force upper case for text */
-	UPROPERTY(EditAnywhere, Getter = "GetEnforceUpperCase", Setter = "SetEnforceUpperCase", Category = "Text", meta = (AllowPrivateAccess="true"))
+	UPROPERTY(EditAnywhere, Getter="GetEnforceUpperCase", Setter="SetEnforceUpperCase", BlueprintSetter="SetEnforceUpperCase", BlueprintGetter="GetEnforceUpperCase", Category = "Text", meta = (AllowPrivateAccess="true"))
 	bool bEnforceUpperCase = false;
 
 	UPROPERTY(VisibleInstanceOnly, Instanced, Category="Token", meta=(AllowPrivateAccess="true"))
@@ -561,6 +565,8 @@ private:
 
 	void OnComponentTransformed(USceneComponent*, EUpdateTransformFlags, ETeleportType);
 
+	bool IsCompiledFromBlueprint() const;
+
 	bool bZeroScaled = false;
 
 #if WITH_EDITORONLY_DATA
