MotionDesign : Text3D
- Fixed text characters array being filled with nullptr instead of character data when used in a spawnable and causing text to not render
- Fixed 5.5 level that contains a Text3DComponent in a BP crashes the engine when opened in 5.6 (and up)
- Fixed invalid character extension would cause a crash for older level

#FailedValidations Preflight

diff --git a/Engine/Plugins/VirtualProduction/Text3D/Source/Text3D/Private/Extensions/Text3DDefaultLayoutExtension.cpp b/Engine/Plugins/VirtualProduction/Text3D/Source/Text3D/Private/Extensions/Text3DDefaultLayoutExtension.cpp
--- a/Engine/Plugins/VirtualProduction/Text3D/Source/Text3D/Private/Extensions/Text3DDefaultLayoutExtension.cpp
+++ b/Engine/Plugins/VirtualProduction/Text3D/Source/Text3D/Private/Extensions/Text3DDefaultLayoutExtension.cpp
@@ -234,8 +234,15 @@ EText3DExtensionResult UText3DDefaultLayoutExtension::PreRendererUpdate(const UE
 
 	if (EnumHasAnyFlags(InParameters.UpdateFlags, EText3DRendererFlags::Geometry))
 	{
-		UText3DCharacterExtensionBase* CharacterExtension = Text3DComponent->GetCharacterExtension();
-		CharacterExtension->AllocateCharacters(CharacterCount);
+		if (UText3DCharacterExtensionBase* CharacterExtension = Text3DComponent->GetCharacterExtension())
+		{
+			CharacterExtension->AllocateCharacters(CharacterCount);
+		}
+		else
+		{
+			UE_LOG(LogText3D, Warning, TEXT("Text3D character extension is invalid, cannot proceed without it"))
+			return EText3DExtensionResult::Failed;
+		}
 	}
 
 	GlyphText->Kernings.SetNum(CharacterCount);
diff --git a/Engine/Plugins/VirtualProduction/Text3D/Source/Text3D/Private/Renderers/DynamicMesh/Text3DDynamicMeshRenderer.cpp b/Engine/Plugins/VirtualProduction/Text3D/Source/Text3D/Private/Renderers/DynamicMesh/Text3DDynamicMeshRenderer.cpp
--- a/Engine/Plugins/VirtualProduction/Text3D/Source/Text3D/Private/Renderers/DynamicMesh/Text3DDynamicMeshRenderer.cpp
+++ b/Engine/Plugins/VirtualProduction/Text3D/Source/Text3D/Private/Renderers/DynamicMesh/Text3DDynamicMeshRenderer.cpp
@@ -200,20 +200,6 @@ void UText3DDynamicMeshRenderer::OnDestroy()
 		DynamicMeshComponent->DestroyComponent();
 		DynamicMeshComponent = nullptr;
 	}
-
-	if (const UText3DComponent* TextComponent = GetText3DComponent())
-	{
-		TArray<USceneComponent*> ChildrenComponents;
-		TextComponent->GetChildrenComponents(/** IncludeChidren */true, ChildrenComponents);
-
-		for (USceneComponent* ChildComponent : ChildrenComponents)
-		{
-			if (ChildComponent->GetName() == TEXT("Text3DDynamicMesh"))
-			{
-				ChildComponent->DestroyComponent();
-			}
-		}
-	}
 }
 
 FName UText3DDynamicMeshRenderer::GetFriendlyName() const
diff --git a/Engine/Plugins/VirtualProduction/Text3D/Source/Text3D/Private/Renderers/DynamicMesh/Text3DDynamicMeshRenderer.h b/Engine/Plugins/VirtualProduction/Text3D/Source/Text3D/Private/Renderers/DynamicMesh/Text3DDynamicMeshRenderer.h
--- a/Engine/Plugins/VirtualProduction/Text3D/Source/Text3D/Private/Renderers/DynamicMesh/Text3DDynamicMeshRenderer.h
+++ b/Engine/Plugins/VirtualProduction/Text3D/Source/Text3D/Private/Renderers/DynamicMesh/Text3DDynamicMeshRenderer.h
@@ -81,7 +81,7 @@ protected:
 	void SaveMesh(const FScopedMeshEditor& InMeshEditor);
 	void RestoreMesh(FScopedMeshEditor& InMeshEditor);
 
-	UPROPERTY(VisibleAnywhere, DuplicateTransient, TextExportTransient, Category="Renderer", meta=(EditCondition="false", EditConditionHides))
+	UPROPERTY(VisibleAnywhere, Instanced, Category="Renderer", meta=(EditCondition="false", EditConditionHides))
 	TObjectPtr<UDynamicMeshComponent> DynamicMeshComponent;
 
 	TArray<FGlyphMeshData> GlyphMeshData;
diff --git a/Engine/Plugins/VirtualProduction/Text3D/Source/Text3D/Private/Renderers/StaticMeshes/Text3DStaticMeshesRenderer.cpp b/Engine/Plugins/VirtualProduction/Text3D/Source/Text3D/Private/Renderers/StaticMeshes/Text3DStaticMeshesRenderer.cpp
--- a/Engine/Plugins/VirtualProduction/Text3D/Source/Text3D/Private/Renderers/StaticMeshes/Text3DStaticMeshesRenderer.cpp
+++ b/Engine/Plugins/VirtualProduction/Text3D/Source/Text3D/Private/Renderers/StaticMeshes/Text3DStaticMeshesRenderer.cpp
@@ -37,6 +37,11 @@ void UText3DStaticMeshesRenderer::OnCreate()
 		TextRoot->DetachFromComponent(DetachRule);
 		const FAttachmentTransformRules AttachRule = FAttachmentTransformRules::KeepRelativeTransform;
 		TextRoot->AttachToComponent(TextComponent, AttachRule);
+
+		if (TextRoot->GetOwner() != TextComponent->GetOwner())
+		{
+			TextRoot->Rename(nullptr, TextComponent->GetOwner(), REN_DontCreateRedirectors | REN_NonTransactional | REN_DoNotDirty);
+		}
 	}
 
 	AllocateComponents(0);
@@ -184,21 +189,6 @@ void UText3DStaticMeshesRenderer::OnDestroy()
 		TextRoot->DestroyComponent();
 		TextRoot = nullptr;
 	}
-
-	if (const UText3DComponent* TextComponent = GetText3DComponent())
-	{
-		TArray<USceneComponent*> ChildrenComponents;
-		TextComponent->GetChildrenComponents(/** IncludeChidren */true, ChildrenComponents);
-
-		for (USceneComponent* ChildComponent : ChildrenComponents)
-		{
-			if (ChildComponent->GetName() == TEXT("TextRoot")
-				|| ChildComponent->GetName().StartsWith(TEXT("CharacterMesh")))
-			{
-				ChildComponent->DestroyComponent();
-			}
-		}
-	}
 }
 
 FName UText3DStaticMeshesRenderer::GetFriendlyName() const
@@ -249,7 +239,11 @@ void UText3DStaticMeshesRenderer::AllocateComponents(int32 InCount)
 				}
 				Actor->RemoveInstanceComponent(MeshComponent);
 				Actor->RemoveOwnedComponent(MeshComponent);
-				CharacterMeshesPool.Add(MeshComponent);
+
+				if (MeshComponent->GetOwner() == Actor)
+				{
+					CharacterMeshesPool.Add(MeshComponent);
+				}
 			}
 		}
 	}
@@ -269,11 +263,13 @@ void UText3DStaticMeshesRenderer::AllocateComponents(int32 InCount)
 				MeshComponent = NewObject<UStaticMeshComponent>(this, StaticMeshComponentName);
 			}
 
-			CharacterMeshes.Add(MeshComponent);
-
-			Actor->AddOwnedComponent(MeshComponent);
-			MeshComponent->RegisterComponent();
-			MeshComponent->AttachToComponent(TextRoot, FAttachmentTransformRules::KeepRelativeTransform);
+			if (MeshComponent)
+			{
+				CharacterMeshes.Add(MeshComponent);
+				Actor->AddOwnedComponent(MeshComponent);
+				MeshComponent->RegisterComponent();
+				MeshComponent->AttachToComponent(TextRoot, FAttachmentTransformRules::KeepRelativeTransform);
+			}
 		}
 	}
 
diff --git a/Engine/Plugins/VirtualProduction/Text3D/Source/Text3D/Private/Renderers/StaticMeshes/Text3DStaticMeshesRenderer.h b/Engine/Plugins/VirtualProduction/Text3D/Source/Text3D/Private/Renderers/StaticMeshes/Text3DStaticMeshesRenderer.h
--- a/Engine/Plugins/VirtualProduction/Text3D/Source/Text3D/Private/Renderers/StaticMeshes/Text3DStaticMeshesRenderer.h
+++ b/Engine/Plugins/VirtualProduction/Text3D/Source/Text3D/Private/Renderers/StaticMeshes/Text3DStaticMeshesRenderer.h
@@ -65,11 +65,11 @@ protected:
 #endif
 
 	/** Holds all character components */
-	UPROPERTY(VisibleAnywhere, Category="Renderer", meta=(EditCondition="false", EditConditionHides))
+	UPROPERTY(VisibleAnywhere, Instanced, Category="Renderer", meta=(EditCondition="false", EditConditionHides))
 	TObjectPtr<USceneComponent> TextRoot;
 
 	/** Each character mesh is held in these components */
-	UPROPERTY(VisibleAnywhere, Category="Renderer", meta=(EditCondition="false", EditConditionHides))
+	UPROPERTY(VisibleAnywhere, Instanced, Category="Renderer", meta=(EditCondition="false", EditConditionHides))
 	TArray<TObjectPtr<UStaticMeshComponent>> CharacterMeshes;
 
 	/** Reusable components to avoid creating and destroying them, recycle components or create new ones when pool is empty */
diff --git a/Engine/Plugins/VirtualProduction/Text3D/Source/Text3D/Private/Renderers/Text3DRendererBase.cpp b/Engine/Plugins/VirtualProduction/Text3D/Source/Text3D/Private/Renderers/Text3DRendererBase.cpp
--- a/Engine/Plugins/VirtualProduction/Text3D/Source/Text3D/Private/Renderers/Text3DRendererBase.cpp
+++ b/Engine/Plugins/VirtualProduction/Text3D/Source/Text3D/Private/Renderers/Text3DRendererBase.cpp
@@ -250,9 +250,3 @@ FBox UText3DRendererBase::GetBounds() const
 
 	return CachedBounds.GetValue();
 }
-
-void UText3DRendererBase::BeginDestroy()
-{
-	Super::BeginDestroy();
-	Destroy();
-}
diff --git a/Engine/Plugins/VirtualProduction/Text3D/Source/Text3D/Private/Text3DComponent.cpp b/Engine/Plugins/VirtualProduction/Text3D/Source/Text3D/Private/Text3DComponent.cpp
--- a/Engine/Plugins/VirtualProduction/Text3D/Source/Text3D/Private/Text3DComponent.cpp
+++ b/Engine/Plugins/VirtualProduction/Text3D/Source/Text3D/Private/Text3DComponent.cpp
@@ -57,7 +57,7 @@ UText3DComponent::UText3DComponent()
 	RenderingExtension = CreateDefaultSubobject<UText3DRenderingExtensionBase, UText3DDefaultRenderingExtension>(TEXT("RenderingExtension"));
 	CharacterExtension = CreateDefaultSubobject<UText3DCharacterExtensionBase, UText3DDefaultCharacterExtension>(TEXT("CharacterExtension"));
 
-	if (!IsTemplate())
+	if (!IsTemplate(RF_ClassDefaultObject))
 	{
 		TransformUpdated.AddUObject(this, &UText3DComponent::OnComponentTransformed);
 	}
@@ -78,11 +78,6 @@ void UText3DComponent::OnComponentDestroyed(bool bInDestroyingHierarchy)
 		TextRenderer->Destroy();
 	}
 
-	if (CharacterExtension)
-	{
-		CharacterExtension->AllocateCharacters(0);
-	}
-
 	TransformUpdated.RemoveAll(this);
 
 	FTSTicker::GetCoreTicker().RemoveTicker(TextUpdateHandle);
@@ -99,7 +94,7 @@ void UText3DComponent::PostLoad()
 
 void UText3DComponent::RequestUpdate(EText3DRendererFlags InFlags, bool bInImmediate)
 {
-	if (InFlags != EText3DRendererFlags::None && !IsTemplate())
+	if (InFlags != EText3DRendererFlags::None && !IsTemplate(RF_ClassDefaultObject))
 	{
 		if (bInImmediate)
 		{
@@ -1061,7 +1056,7 @@ void UText3DComponent::Rebuild(EText3DRendererFlags InUpdateFlags)
 
 		if (!TextRenderer)
 		{
-			UE_LOG(LogText3D, Error, TEXT("Text3D renderer is not valid, cannot proceed"));
+			UE_LOG(LogText3D, Warning, TEXT("Text3D renderer is not valid, cannot proceed"));
 			return;
 		}
 	}
@@ -1075,7 +1070,7 @@ void UText3DComponent::Rebuild(EText3DRendererFlags InUpdateFlags)
 
 		if (!Font)
 		{
-			UE_LOG(LogText3D, Error, TEXT("Font object is not valid, Fallback font is not defined in project settings, cannot proceed"));
+			UE_LOG(LogText3D, Warning, TEXT("Font object is not valid, Fallback font is not defined in project settings, cannot proceed"));
 			return;
 		}
 	}
@@ -1152,7 +1147,7 @@ void UText3DComponent::Rebuild(EText3DRendererFlags InUpdateFlags)
 
 void UText3DComponent::OnTextRendererClassChanged()
 {
-	if (IsTemplate())
+	if (IsTemplate(RF_ClassDefaultObject))
 	{
 		return;
 	}
diff --git a/Engine/Plugins/VirtualProduction/Text3D/Source/Text3D/Public/Characters/Text3DCharacterBase.h b/Engine/Plugins/VirtualProduction/Text3D/Source/Text3D/Public/Characters/Text3DCharacterBase.h
--- a/Engine/Plugins/VirtualProduction/Text3D/Source/Text3D/Public/Characters/Text3DCharacterBase.h
+++ b/Engine/Plugins/VirtualProduction/Text3D/Source/Text3D/Public/Characters/Text3DCharacterBase.h
@@ -94,20 +94,20 @@ protected:
 
 #if WITH_EDITORONLY_DATA
 	/** Used for TitleProperty of the array to display this instead of the index */
-	UPROPERTY(VisibleInstanceOnly, Category="Character", Transient, meta=(EditCondition="false", EditConditionHides, AllowPrivateAccess="true"))
+	UPROPERTY(VisibleAnywhere, Category="Character", Transient, meta=(EditCondition="false", EditConditionHides, AllowPrivateAccess="true"))
 	FString Character;
 #endif
 
-	UPROPERTY(EditInstanceOnly, Setter, Getter, Category="Character", meta=(AllowPrivateAccess = "true"))
+	UPROPERTY(EditAnywhere, Setter, Getter, Category="Character", meta=(AllowPrivateAccess = "true"))
 	FVector RelativeLocation = FVector::ZeroVector;
 	
-	UPROPERTY(EditInstanceOnly, Setter, Getter, Category="Character", meta=(AllowPrivateAccess = "true"))
+	UPROPERTY(EditAnywhere, Setter, Getter, Category="Character", meta=(AllowPrivateAccess = "true"))
 	FRotator RelativeRotation = FRotator::ZeroRotator;
 	
-	UPROPERTY(EditInstanceOnly, Setter, Getter, Category="Character", meta=(Delta="0.05", MotionDesignVectorWidget, AllowPreserveRatio="XYZ", ClampMin="0", AllowPrivateAccess = "true"))
+	UPROPERTY(EditAnywhere, Setter, Getter, Category="Character", meta=(Delta="0.05", MotionDesignVectorWidget, AllowPreserveRatio="XYZ", ClampMin="0", AllowPrivateAccess = "true"))
 	FVector RelativeScale = FVector::OneVector;
 
-	UPROPERTY(EditInstanceOnly, Setter="SetVisibility", Getter="GetVisibility", Category="Character", meta=(AllowPrivateAccess = "true"))
+	UPROPERTY(EditAnywhere, Setter="SetVisibility", Getter="GetVisibility", Category="Character", meta=(AllowPrivateAccess = "true"))
 	bool bVisible = true;
 
 	/** Tag used to identify the custom style applied to this character */
diff --git a/Engine/Plugins/VirtualProduction/Text3D/Source/Text3D/Public/Characters/Text3DDefaultCharacter.h b/Engine/Plugins/VirtualProduction/Text3D/Source/Text3D/Public/Characters/Text3DDefaultCharacter.h
--- a/Engine/Plugins/VirtualProduction/Text3D/Source/Text3D/Public/Characters/Text3DDefaultCharacter.h
+++ b/Engine/Plugins/VirtualProduction/Text3D/Source/Text3D/Public/Characters/Text3DDefaultCharacter.h
@@ -34,6 +34,6 @@ protected:
 	//~ End UText3DCharacterBase
 
 	/** Kerning adjusts the space between this character */
-	UPROPERTY(EditInstanceOnly, Getter, Setter, Category="Character", meta=(AllowPrivateAccess = "true"))
+	UPROPERTY(EditAnywhere, Getter, Setter, Category="Character", meta=(AllowPrivateAccess = "true"))
 	float Kerning = 0.f;
 };
diff --git a/Engine/Plugins/VirtualProduction/Text3D/Source/Text3D/Public/Extensions/Text3DDefaultCharacterExtension.h b/Engine/Plugins/VirtualProduction/Text3D/Source/Text3D/Public/Extensions/Text3DDefaultCharacterExtension.h
--- a/Engine/Plugins/VirtualProduction/Text3D/Source/Text3D/Public/Extensions/Text3DDefaultCharacterExtension.h
+++ b/Engine/Plugins/VirtualProduction/Text3D/Source/Text3D/Public/Extensions/Text3DDefaultCharacterExtension.h
@@ -24,7 +24,7 @@ protected:
 	void AllocateTextCharacters(uint16 InCharacterCount);
 
 	/** Characters composing the active text */
-	UPROPERTY()
+	UPROPERTY(Instanced)
 	TArray<TObjectPtr<UText3DCharacterBase>> TextCharacters;
 
 	/** Pool of characters to reuse when text changes to avoid creating new objects */
diff --git a/Engine/Plugins/VirtualProduction/Text3D/Source/Text3D/Public/Renderers/Text3DRendererBase.h b/Engine/Plugins/VirtualProduction/Text3D/Source/Text3D/Public/Renderers/Text3DRendererBase.h
--- a/Engine/Plugins/VirtualProduction/Text3D/Source/Text3D/Public/Renderers/Text3DRendererBase.h
+++ b/Engine/Plugins/VirtualProduction/Text3D/Source/Text3D/Public/Renderers/Text3DRendererBase.h
@@ -43,10 +43,6 @@ public:
 	}
 
 protected:
-	//~ Begin UObject
-	virtual void BeginDestroy() override;
-	//~ End UObject
-
 	/** Create and setup the implementation components, called on load or creation */
 	virtual void OnCreate() PURE_VIRTUAL(UText3DRenderingImplementationBase::OnCreate, );
 
diff --git a/Engine/Plugins/VirtualProduction/Text3D/Source/Text3D/Public/Text3DComponent.h b/Engine/Plugins/VirtualProduction/Text3D/Source/Text3D/Public/Text3DComponent.h
--- a/Engine/Plugins/VirtualProduction/Text3D/Source/Text3D/Public/Text3DComponent.h
+++ b/Engine/Plugins/VirtualProduction/Text3D/Source/Text3D/Public/Text3DComponent.h
@@ -499,28 +499,28 @@ protected:
 	UPROPERTY(EditAnywhere, Getter="GetEnforceUpperCase", Setter="SetEnforceUpperCase", BlueprintSetter="SetEnforceUpperCase", BlueprintGetter="GetEnforceUpperCase", Category = "Text", meta = (AllowPrivateAccess="true"))
 	bool bEnforceUpperCase = false;
 
-	UPROPERTY(VisibleInstanceOnly, Instanced, Category="Token", meta=(AllowPrivateAccess="true"))
+	UPROPERTY(VisibleAnywhere, Instanced, Category="Token", meta=(AllowPrivateAccess="true"))
 	TObjectPtr<UText3DTokenExtensionBase> TokenExtension;
 
-	UPROPERTY(VisibleInstanceOnly, Instanced, Category="Style", meta=(AllowPrivateAccess="true"))
+	UPROPERTY(VisibleAnywhere, Instanced, Category="Style", meta=(AllowPrivateAccess="true"))
 	TObjectPtr<UText3DStyleExtensionBase> StyleExtension;
 
-	UPROPERTY(VisibleInstanceOnly, Instanced, Category="Character", meta=(AllowPrivateAccess="true"))
+	UPROPERTY(VisibleAnywhere, Instanced, Category="Character", meta=(AllowPrivateAccess="true"))
 	TObjectPtr<UText3DCharacterExtensionBase> CharacterExtension;
 
-	UPROPERTY(VisibleInstanceOnly, Instanced, Category="Layout", meta=(AllowPrivateAccess="true"))
+	UPROPERTY(VisibleAnywhere, Instanced, Category="Layout", meta=(AllowPrivateAccess="true"))
 	TObjectPtr<UText3DLayoutExtensionBase> LayoutExtension;
 
-	UPROPERTY(EditInstanceOnly, Instanced, Category="LayoutEffects", meta=(AllowPrivateAccess="true"))
+	UPROPERTY(EditAnywhere, Instanced, Category="LayoutEffects", meta=(AllowPrivateAccess="true"))
 	TArray<TObjectPtr<UText3DLayoutEffectBase>> LayoutEffects;
 
-	UPROPERTY(VisibleInstanceOnly, Instanced, Category="Geometry", meta=(AllowPrivateAccess="true"))
+	UPROPERTY(VisibleAnywhere, Instanced, Category="Geometry", meta=(AllowPrivateAccess="true"))
 	TObjectPtr<UText3DGeometryExtensionBase> GeometryExtension;
 
-	UPROPERTY(VisibleInstanceOnly, Instanced, Category="Material", meta=(AllowPrivateAccess="true"))
+	UPROPERTY(VisibleAnywhere, Instanced, Category="Material", meta=(AllowPrivateAccess="true"))
 	TObjectPtr<UText3DMaterialExtensionBase> MaterialExtension;
 
-	UPROPERTY(VisibleInstanceOnly, Instanced, Category="Rendering", meta=(AllowPrivateAccess="true"))
+	UPROPERTY(VisibleAnywhere, Instanced, Category="Rendering", meta=(AllowPrivateAccess="true"))
 	TObjectPtr<UText3DRenderingExtensionBase> RenderingExtension;
 
 private:
