MotionDesign :
- Various minimal fixes for ActorModifier, PropertyAnimator and Text3D

diff --git a/Engine/Plugins/VirtualProduction/Avalanche/Source/AvalancheMask/Private/Mask2D/AvaMask2DBaseModifier.cpp b/Engine/Plugins/VirtualProduction/Avalanche/Source/AvalancheMask/Private/Mask2D/AvaMask2DBaseModifier.cpp
--- a/Engine/Plugins/VirtualProduction/Avalanche/Source/AvalancheMask/Private/Mask2D/AvaMask2DBaseModifier.cpp
+++ b/Engine/Plugins/VirtualProduction/Avalanche/Source/AvalancheMask/Private/Mask2D/AvaMask2DBaseModifier.cpp
@@ -919,7 +919,7 @@ void UAvaMask2DBaseModifier::OnMaterialsChanged(UObject* InMaterialOwner, const
 				FInstancedStruct* HandleData = MaterialCollectionHandleData.Find(OwningActor);
 				const TSharedPtr<IAvaMaskMaterialCollectionHandle>* CollectionHandle = MaterialCollectionHandles.Find(OwningActor);
 
-				if (HandleData && CollectionHandle)
+				if (HandleData && CollectionHandle && CollectionHandle->IsValid())
 				{
 					(*CollectionHandle)->SaveOriginalState(*HandleData);
 				}
diff --git a/Engine/Plugins/VirtualProduction/PropertyAnimatorCore/Source/PropertyAnimatorCore/Public/Animators/PropertyAnimatorCoreBase.h b/Engine/Plugins/VirtualProduction/PropertyAnimatorCore/Source/PropertyAnimatorCore/Public/Animators/PropertyAnimatorCoreBase.h
--- a/Engine/Plugins/VirtualProduction/PropertyAnimatorCore/Source/PropertyAnimatorCore/Public/Animators/PropertyAnimatorCoreBase.h
+++ b/Engine/Plugins/VirtualProduction/PropertyAnimatorCore/Source/PropertyAnimatorCore/Public/Animators/PropertyAnimatorCoreBase.h
@@ -340,7 +340,7 @@ private:
 	/** Needed to restore properties and stop animation before world is saved */
 	void OnPreSaveWorld(UWorld* InWorld, FObjectPreSaveContext InContext);
 
-	/** Needed to restore properties and stop animation before PIE is started */
+	/** Needed to restore properties and reset animation before PIE is started */
 	void OnPreBeginPIE(bool bInSimulating);
 #endif
 
diff --git a/Engine/Plugins/VirtualProduction/Text3D/Source/Text3D/Private/Extensions/Text3DDefaultMaterialExtension.cpp b/Engine/Plugins/VirtualProduction/Text3D/Source/Text3D/Private/Extensions/Text3DDefaultMaterialExtension.cpp
--- a/Engine/Plugins/VirtualProduction/Text3D/Source/Text3D/Private/Extensions/Text3DDefaultMaterialExtension.cpp
+++ b/Engine/Plugins/VirtualProduction/Text3D/Source/Text3D/Private/Extensions/Text3DDefaultMaterialExtension.cpp
@@ -346,7 +346,7 @@ EText3DExtensionResult UText3DDefaultMaterialExtension::PreRendererUpdate(const
 				TextBounds = TextBounds.TransformBy(ComponentTransform.Inverse());
 				// Apply scale ratio based on font size 64 => 1.f
 				FTransform RatioTransform;
-				RatioTransform.SetScale3D(FVector(64.f/Text3DComponent->GetFontSize()));
+				RatioTransform.SetScale3D(FVector(64.f / FMath::Max(Text3DComponent->GetFontSize(), 0.1f)));
 				LocalTextBounds = TextBounds.TransformBy(RatioTransform);
 			}
 
diff --git a/Engine/Plugins/VirtualProduction/Text3D/Source/Text3D/Public/Text3DComponent.h b/Engine/Plugins/VirtualProduction/Text3D/Source/Text3D/Public/Text3DComponent.h
--- a/Engine/Plugins/VirtualProduction/Text3D/Source/Text3D/Public/Text3DComponent.h
+++ b/Engine/Plugins/VirtualProduction/Text3D/Source/Text3D/Public/Text3DComponent.h
@@ -492,7 +492,7 @@ protected:
 	FName Typeface;
 
 	/** Text font size defines the size of each rendered character */
-	UPROPERTY(EditAnywhere, Getter, Setter, Category = "Text", meta=(ClampMin = "0.1", AllowPrivateAccess = "true"))
+	UPROPERTY(EditAnywhere, Getter, Setter, DisplayName = "Default FontSize", Category = "Text", meta=(ClampMin = "0.1", AllowPrivateAccess = "true"))
 	float FontSize = 48.f;
 
 	/** Whether to force upper case for text */
