MotionDesign : Text3D
- Fixed crash caused by an invalid character object/mesh, the causes are unknown due to missing repro steps, but this will add additional checks for safety

diff --git a/Engine/Plugins/VirtualProduction/Text3D/Source/Text3D/Private/Extensions/Text3DDefaultLayoutExtension.cpp b/Engine/Plugins/VirtualProduction/Text3D/Source/Text3D/Private/Extensions/Text3DDefaultLayoutExtension.cpp
--- a/Engine/Plugins/VirtualProduction/Text3D/Source/Text3D/Private/Extensions/Text3DDefaultLayoutExtension.cpp
+++ b/Engine/Plugins/VirtualProduction/Text3D/Source/Text3D/Private/Extensions/Text3DDefaultLayoutExtension.cpp
@@ -270,7 +270,7 @@ EText3DExtensionResult UText3DDefaultLayoutExtension::PreRendererUpdate(const UE
 				}
 				else
 				{
-					UE_LOG(LogText3D, Warning, TEXT("Invalid character object returned for index %i, cannot proceed"), CharacterIndex)
+					UE_LOG(LogText3D, Warning, TEXT("Invalid character object returned for index %i before layout calculations, cannot proceed"), CharacterIndex)
 					return EText3DExtensionResult::Failed;
 				}
 			}
@@ -302,6 +302,12 @@ EText3DExtensionResult UText3DDefaultLayoutExtension::PreRendererUpdate(const UE
 				const uint16 Index = bIsRTL ? CharacterCount - 1 - CharacterIndex : CharacterIndex;
 				UText3DCharacterBase* Character = Text3DComponent->GetCharacter(Index);
 
+				if (!Character)
+				{
+					UE_LOG(LogText3D, Warning, TEXT("Invalid character object returned for index %i after layout calculations, cannot proceed"), CharacterIndex)
+					return EText3DExtensionResult::Failed;
+				}
+
 				if (EnumHasAnyFlags(InParameters.UpdateFlags, EText3DRendererFlags::Geometry))
 				{
 #if WITH_FREETYPE
@@ -323,11 +329,19 @@ EText3DExtensionResult UText3DDefaultLayoutExtension::PreRendererUpdate(const UE
 					}
 					else
 					{
-						UE_LOG(LogText3D, Error, TEXT("Failed to get cached font data for '%u %i' in Text3D geometry extension"), FontFaceHash, GlyphEntry->Entry.GlyphIndex);
+						UE_LOG(LogText3D, Warning, TEXT("Failed to get cached font data for '%u %i' in Text3D geometry extension"), FontFaceHash, GlyphEntry->Entry.GlyphIndex);
+						return EText3DExtensionResult::Failed;
 					}
 				}
 
 				const FText3DCachedMesh* CachedMesh = Character->GetGlyphMesh();
+
+				if (!CachedMesh)
+				{
+					UE_LOG(LogText3D, Warning, TEXT("Invalid character cached mesh returned for index %i and glyph %i, cannot proceed"), CharacterIndex, GlyphEntry->Entry.GlyphIndex)
+					return EText3DExtensionResult::Failed;
+				}
+
 				const FBox GlyphBounds = CachedMesh->MeshBounds;
 				const FVector GlyphSize = GlyphBounds.GetSize();
 
