MotionDesign : Text3D
- Fixed compilation error caused by Freetype usage outside macros for an UnrealServer build

diff --git a/Engine/Plugins/VirtualProduction/Text3D/Source/Text3D/Private/Subsystems/Text3DEngineSubsystem.cpp b/Engine/Plugins/VirtualProduction/Text3D/Source/Text3D/Private/Subsystems/Text3DEngineSubsystem.cpp
--- a/Engine/Plugins/VirtualProduction/Text3D/Source/Text3D/Private/Subsystems/Text3DEngineSubsystem.cpp
+++ b/Engine/Plugins/VirtualProduction/Text3D/Source/Text3D/Private/Subsystems/Text3DEngineSubsystem.cpp
@@ -217,10 +217,12 @@ FText3DCachedMesh* FText3DFontFaceCache::FindOrBuildGlyphMesh(uint32 InGlyphInde
 
 	FText3DCachedMesh& CachedMesh = GlyphMeshes.FindOrAdd(HashParameters);
 
+#if WITH_FREETYPE
 	if (const FT_Face Face = FontFace->GetFace())
 	{
 		CachedMesh.FontFaceGlyphSize = FVector2D(Face->size->metrics.x_ppem, Face->size->metrics.y_ppem);
 	}
+#endif
 
 	const FString MeshName = FString::Printf(TEXT("Text3D_Char_%u"), HashParameters);
 	const FName StaticMeshName = MakeUniqueObjectName(TextSubsystem, UStaticMesh::StaticClass(), FName(*MeshName));
