MotionDesign : Text3D
- Fixed infinite loop reported in FText3DGlyphLoader::RemoveBadPoints

diff --git a/Engine/Plugins/VirtualProduction/Text3D/Source/Text3D/Private/GeometryBuilders/Text3DGlyphLoader.cpp b/Engine/Plugins/VirtualProduction/Text3D/Source/Text3D/Private/GeometryBuilders/Text3DGlyphLoader.cpp
--- a/Engine/Plugins/VirtualProduction/Text3D/Source/Text3D/Private/GeometryBuilders/Text3DGlyphLoader.cpp
+++ b/Engine/Plugins/VirtualProduction/Text3D/Source/Text3D/Private/GeometryBuilders/Text3DGlyphLoader.cpp
@@ -588,40 +588,40 @@ bool FText3DGlyphLoader::MergedWithNext(const TSharedPtr<FPolygon2f> Contour, co
 
 TSharedPtr<FPolygon2f> FText3DGlyphLoader::RemoveBadPoints(const TSharedPtr<FPolygon2f> Contour) const
 {
-	int32 First = 0;
-
-	for (int32 Point = First;;)
+	int32 VertexCount = Contour->VertexCount();
+	if (VertexCount < 4)
 	{
-		bool bPointRemoved = false;
+		return nullptr;
+	}
+
+	int32 CheckVertexIndex = 0;
+	int32 ConsecutiveCleanVertices = 0;
 
-		while (IsBadNormal(Contour, Point) || MergedWithNext(Contour, Point))
+	while (ConsecutiveCleanVertices < VertexCount)
+	{
+		if (IsBadNormal(Contour, CheckVertexIndex) || MergedWithNext(Contour, CheckVertexIndex))
 		{
-			bPointRemoved = true;
+			Contour->RemoveVertex(CheckVertexIndex);
 
-			if (Contour->VertexCount() < 4)
+			// Update count since we removed one
+			VertexCount = Contour->VertexCount();
+			if (VertexCount < 4)
 			{
-				return TSharedPtr<FPolygon2f>();
+				return nullptr;
 			}
 
-			Contour->RemoveVertex(Point);
-			const int32 Current = (Point - 1 + Contour->VertexCount()) % Contour->VertexCount();
-
-			if (Point == 0)
-			{
-				First = Current;
-			}
+			// After removal, the next vertex now occupies CheckVertexIndex
+			// Wrap around if we removed the last vertex
+			CheckVertexIndex = CheckVertexIndex % VertexCount;
 
-			Point = Current;
+			// Reset counter since we made a change
+			ConsecutiveCleanVertices = 0;
 		}
-
-		if (!bPointRemoved)
+		else
 		{
-			Point = (Point + 1) % Contour->VertexCount();
-
-			if (Point == First)
-			{
-				break;
-			}
+			// This vertex was clean, move to next
+			ConsecutiveCleanVertices++;
+			CheckVertexIndex = (CheckVertexIndex + 1) % VertexCount;
 		}
 	}
 
