MotionDesign : Text3D
- Fixed Text3D used in blueprint would apply additive transform if parent component is moved

diff --git a/Engine/Plugins/VirtualProduction/Text3D/Source/Text3D/Private/Renderers/DynamicMesh/Text3DDynamicMeshRenderer.cpp b/Engine/Plugins/VirtualProduction/Text3D/Source/Text3D/Private/Renderers/DynamicMesh/Text3DDynamicMeshRenderer.cpp
--- a/Engine/Plugins/VirtualProduction/Text3D/Source/Text3D/Private/Renderers/DynamicMesh/Text3DDynamicMeshRenderer.cpp
+++ b/Engine/Plugins/VirtualProduction/Text3D/Source/Text3D/Private/Renderers/DynamicMesh/Text3DDynamicMeshRenderer.cpp
@@ -53,11 +53,15 @@ void UText3DDynamicMeshRenderer::FScopedMeshEditor::ProcessMesh(const TFunctionR
 void UText3DDynamicMeshRenderer::OnCreate()
 {
 	UText3DComponent* TextComponent = GetText3DComponent();
+	AActor* Owner = TextComponent->GetOwner();
+	const FAttachmentTransformRules AttachRule = FAttachmentTransformRules::SnapToTargetIncludingScale;
 
 	if (!DynamicMeshComponent)
 	{
 		DynamicMeshComponent = NewObject<UDynamicMeshComponent>(this, TEXT("Text3DDynamicMesh"));
-		DynamicMeshComponent->SetupAttachment(TextComponent);
+		DynamicMeshComponent->RegisterComponent();
+		Owner->AddOwnedComponent(DynamicMeshComponent);
+		DynamicMeshComponent->AttachToComponent(TextComponent, AttachRule);
 
 		// Setup collision
 		DynamicMeshComponent->SetComplexAsSimpleCollisionEnabled(true);
@@ -71,7 +75,6 @@ void UText3DDynamicMeshRenderer::OnCreate()
 		FDetachmentTransformRules DetachRule = FDetachmentTransformRules::KeepRelativeTransform;
 		DetachRule.bCallModify = false;
 		DynamicMeshComponent->DetachFromComponent(DetachRule);
-		const FAttachmentTransformRules AttachRule = FAttachmentTransformRules::KeepRelativeTransform;
 		DynamicMeshComponent->AttachToComponent(TextComponent, AttachRule);
 	}
 }
diff --git a/Engine/Plugins/VirtualProduction/Text3D/Source/Text3D/Private/Renderers/StaticMeshes/Text3DStaticMeshesRenderer.cpp b/Engine/Plugins/VirtualProduction/Text3D/Source/Text3D/Private/Renderers/StaticMeshes/Text3DStaticMeshesRenderer.cpp
--- a/Engine/Plugins/VirtualProduction/Text3D/Source/Text3D/Private/Renderers/StaticMeshes/Text3DStaticMeshesRenderer.cpp
+++ b/Engine/Plugins/VirtualProduction/Text3D/Source/Text3D/Private/Renderers/StaticMeshes/Text3DStaticMeshesRenderer.cpp
@@ -24,23 +24,26 @@
 void UText3DStaticMeshesRenderer::OnCreate()
 {
 	UText3DComponent* TextComponent = GetText3DComponent();
-
+	AActor* Owner = TextComponent->GetOwner();
+	const FAttachmentTransformRules AttachRule = FAttachmentTransformRules::SnapToTargetIncludingScale;
+	
 	if (!TextRoot)
 	{
 		TextRoot = NewObject<USceneComponent>(this, TEXT("TextRoot"));
-		TextRoot->SetupAttachment(TextComponent);
+		TextRoot->RegisterComponent();
+		Owner->AddOwnedComponent(TextRoot);
+		TextRoot->AttachToComponent(TextComponent, AttachRule);
 	}
 	else
 	{
 		FDetachmentTransformRules DetachRule = FDetachmentTransformRules::KeepRelativeTransform;
 		DetachRule.bCallModify = false;
 		TextRoot->DetachFromComponent(DetachRule);
-		const FAttachmentTransformRules AttachRule = FAttachmentTransformRules::KeepRelativeTransform;
 		TextRoot->AttachToComponent(TextComponent, AttachRule);
 
-		if (TextRoot->GetOwner() != TextComponent->GetOwner())
+		if (TextRoot->GetOwner() != Owner)
 		{
-			TextRoot->Rename(nullptr, TextComponent->GetOwner(), REN_DontCreateRedirectors | REN_NonTransactional | REN_DoNotDirty);
+			TextRoot->Rename(nullptr, Owner, REN_DontCreateRedirectors | REN_NonTransactional | REN_DoNotDirty);
 		}
 	}
 
@@ -268,7 +271,7 @@ void UText3DStaticMeshesRenderer::AllocateComponents(int32 InCount)
 				CharacterMeshes.Add(MeshComponent);
 				Actor->AddOwnedComponent(MeshComponent);
 				MeshComponent->RegisterComponent();
-				MeshComponent->AttachToComponent(TextRoot, FAttachmentTransformRules::KeepRelativeTransform);
+				MeshComponent->AttachToComponent(TextRoot, FAttachmentTransformRules::SnapToTargetIncludingScale);
 			}
 		}
 	}
