MotionDesign : Text3D
- Fixed mixing RTL and LTR characters would cause characters duplication

diff --git a/Engine/Plugins/VirtualProduction/Text3D/Source/Text3D/Private/Extensions/Text3DDefaultLayoutExtension.cpp b/Engine/Plugins/VirtualProduction/Text3D/Source/Text3D/Private/Extensions/Text3DDefaultLayoutExtension.cpp
--- a/Engine/Plugins/VirtualProduction/Text3D/Source/Text3D/Private/Extensions/Text3DDefaultLayoutExtension.cpp
+++ b/Engine/Plugins/VirtualProduction/Text3D/Source/Text3D/Private/Extensions/Text3DDefaultLayoutExtension.cpp
@@ -254,15 +254,12 @@ EText3DExtensionResult UText3DDefaultLayoutExtension::PreRendererUpdate(const UE
 #if WITH_FREETYPE
 	for (const TConstEnumerateRef<FGlyphLine> GlyphLine : EnumerateRange(GlyphText->Lines))
 	{
-		const bool bIsRTL = (GlyphLine->TextDirection == TextBiDi::ETextDirection::RightToLeft);
-
 		for (const TConstEnumerateRef<FGlyphEntry> GlyphEntry : EnumerateRange(GlyphLine->Glyphs))
 		{
 			if (FGlyphText::IsValidGlyph(*GlyphEntry))
 			{
 				const TSharedPtr<FFreeTypeFace> FontFacePtr = GlyphEntry->Entry.FontFaceData->FontFace.Pin();
-				const uint16 Index = bIsRTL ? CharacterCount - 1 - CharacterIndex : CharacterIndex;
-				if (const UText3DCharacterBase* Character = Text3DComponent->GetCharacter(Index))
+				if (const UText3DCharacterBase* Character = Text3DComponent->GetCharacter(CharacterIndex))
 				{
 					GlyphText->Kernings[CharacterIndex] = Character->GetCharacterKerning();
                     GlyphText->FontFaces[CharacterIndex] = FontFacePtr;
@@ -293,14 +290,11 @@ EText3DExtensionResult UText3DDefaultLayoutExtension::PreRendererUpdate(const UE
 	for (TConstEnumerateRef<FGlyphLine> GlyphLine : EnumerateRange(GlyphText->Lines))
 	{
 		FVector Location = GetLineLocation(GlyphLine.GetIndex());
-		const bool bIsRTL = (GlyphLine->TextDirection == TextBiDi::ETextDirection::RightToLeft);
-
 		for (TConstEnumerateRef<FGlyphEntry> GlyphEntry : EnumerateRange(GlyphLine->Glyphs))
 		{
 			if (FGlyphText::IsValidGlyph(*GlyphEntry))
 			{
-				const uint16 Index = bIsRTL ? CharacterCount - 1 - CharacterIndex : CharacterIndex;
-				UText3DCharacterBase* Character = Text3DComponent->GetCharacter(Index);
+				UText3DCharacterBase* Character = Text3DComponent->GetCharacter(CharacterIndex);
 
 				if (!Character)
 				{
diff --git a/Engine/Plugins/VirtualProduction/Text3D/Source/Text3D/Private/LayoutBuilders/Text3DLayoutShaper.cpp b/Engine/Plugins/VirtualProduction/Text3D/Source/Text3D/Private/LayoutBuilders/Text3DLayoutShaper.cpp
--- a/Engine/Plugins/VirtualProduction/Text3D/Source/Text3D/Private/LayoutBuilders/Text3DLayoutShaper.cpp
+++ b/Engine/Plugins/VirtualProduction/Text3D/Source/Text3D/Private/LayoutBuilders/Text3DLayoutShaper.cpp
@@ -59,17 +59,18 @@ void FText3DLayoutShaper::ShapeBidirectionalText(
 
 		for (const TSharedRef<ILayoutBlock>& Block : Line.Blocks)
 		{
+			const FTextRange RunRange = Block->GetRun()->GetTextRange();
+			const FTextRange BlockRange = Block->GetTextRange();
 			FString BlockText;
 			Block->GetRun()->AppendTextTo(BlockText);
 	
 			const FLayoutBlockTextContext BlockTextContext = Block->GetTextContext();
 			const FName StyleTag(Block->GetRun()->GetRunInfo().Name);
 			const FSlateFontInfo& FontInfo = !StyleTag.IsNone() && InStyles->HasWidgetStyle<FTextBlockStyle>(StyleTag) ? InStyles->GetWidgetStyle<FTextBlockStyle>(StyleTag).Font : InDefaultStyle.Font;
-			const FTextRange BlockRange(0, BlockText.Len());
 
 			const FShapedGlyphSequenceRef ShapedGlyphSequence = ShapedTextCacheUtil::GetShapedTextSubSequence(
 				BlockTextContext.ShapedTextCache,
-				FCachedShapedTextKey(BlockRange, 1.0f, BlockTextContext, FontInfo),
+				FCachedShapedTextKey(RunRange, 1.0f, BlockTextContext, FontInfo),
 				BlockRange,
 				*BlockText,
 				BlockTextContext.TextDirection
@@ -82,10 +83,6 @@ void FText3DLayoutShaper::ShapeBidirectionalText(
 				GlyphEntry.Size = FontInfo.Size;
 			}
 		}
-
-		if (ShapedLine.TextDirection == TextBiDi::ETextDirection::RightToLeft)
-		{
-			Algo::Reverse(ShapedLine.Glyphs);
-		}
 	}
 }
+	
\ No newline at end of file
diff --git a/Engine/Plugins/VirtualProduction/Text3D/Source/Text3D/Private/LayoutBuilders/Text3DShapedGlyphText.cpp b/Engine/Plugins/VirtualProduction/Text3D/Source/Text3D/Private/LayoutBuilders/Text3DShapedGlyphText.cpp
--- a/Engine/Plugins/VirtualProduction/Text3D/Source/Text3D/Private/LayoutBuilders/Text3DShapedGlyphText.cpp
+++ b/Engine/Plugins/VirtualProduction/Text3D/Source/Text3D/Private/LayoutBuilders/Text3DShapedGlyphText.cpp
@@ -54,91 +54,109 @@ namespace UE::Text3D::Layout
     		float MaxFontDescender = 0.0f;
 
 	        const int32 GlyphCount = GlyphLine->Glyphs.Num();
-
-			for (TConstEnumerateRef<FGlyphEntry> CurrentGlyph : EnumerateRange(GlyphLine->Glyphs))
-	        {
-    			const bool bVisibleGlyph = IsValidGlyph(*CurrentGlyph);
-				const int32 Index = bIsRTL ? Kernings.Num() - 1 - CharacterIndex : CharacterIndex;
-    			const float CurrentKerning = bVisibleGlyph ? Kernings[Index] : 0.f;
-    			const float GlyphAdv = GlyphLine->GetWidthAdvance(CurrentGlyph.GetIndex(), Tracking + CurrentKerning, WordSpacing);
-
-	            // If we're at the end the line or at whitespace
-	            const bool bWordBreak = !bVisibleGlyph || CurrentGlyph.GetIndex() == GlyphCount - 1;
-
-    			if (bWrap                               // when we're wrapping
-					&& bWordBreak                       // and at a word break
-					&& LineWidth > MaxWidth             // and the current line is longer than the max
-					&& CurrentWordLength != LineWidth)  // and the line is not just a single word that we can't break
-    			{
-    				CurrentLine->Width = LineWidth - CurrentWordLength;
-    				CurrentLine->MaxFontHeight = MaxFontAscender + MaxFontDescender;
-    				CurrentLine->MaxFontAscender = MaxFontAscender;
-    				CurrentLine->MaxFontDescender = MaxFontDescender;
-
-    				// Reset
-    				MaxFontAscender = 0.0f;
-    				MaxFontDescender = 0.0f;
-
-    				CurrentLine = &NewLines.AddDefaulted_GetRef();
-    				CurrentLine->TextDirection = GlyphLine->TextDirection;
-    				LineWidth = CurrentWordLength;
-    			}
-
-				LineWidth += GlyphAdv;
-				CurrentWordLength += GlyphAdv;
-
-				if (bVisibleGlyph)
-				{
-	#if WITH_FREETYPE
-					const FT_Face Face = FontFaces[Index]->GetFace();
-
-					const float Ascender = FMath::Max(0.0f, Face->size->metrics.ascender * UE::Text3D::Metrics::Convert26Dot6ToPixel);
-					const float Descender = FMath::Max(0.0f, -Face->size->metrics.descender * UE::Text3D::Metrics::Convert26Dot6ToPixel);
-
-					MaxFontAscender  = FMath::Max(MaxFontAscender,  Ascender);
-					MaxFontDescender = FMath::Max(MaxFontDescender, Descender);
-	#endif
-				}
-
-				if (bIsRTL)
-				{
-					CurrentWord.Insert(*CurrentGlyph, 0);
-					CurrentAdvances.Insert(GlyphAdv, 0);
-				}
-				else
-				{
-					CurrentWord.Add(*CurrentGlyph);
+	    	
+	    	// Used to avoid an array copy. Here we need to iterate glyphs in logical order instead of visual order...
+	    	auto ForEachGlyph = [&GlyphLine, &bIsRTL](TFunctionRef<void(const FGlyphEntry& InEntry, int32 InGlyphIndex, int32 InIterationCount)> InFunctor)
+	    		{
+	    			int32 ProcessedCount = 0;
+	    			if (bIsRTL)
+	    			{
+	    				for (int32 GlyphIndex = GlyphLine->Glyphs.Num() - 1; GlyphIndex >= 0; GlyphIndex--)
+	    				{
+	    					InFunctor(GlyphLine->Glyphs[GlyphIndex], GlyphIndex, ProcessedCount++);
+	    				}
+	    			}
+				    else
+				    {
+				    	for (int32 GlyphIndex = 0; GlyphIndex < GlyphLine->Glyphs.Num(); GlyphIndex++)
+				    	{
+				    		InFunctor(GlyphLine->Glyphs[GlyphIndex], GlyphIndex, ProcessedCount++);
+				    	}
+				    }
+	    		};
+
+	    	ForEachGlyph(
+	    		[this, &CharacterIndex, &GlyphCount, &GlyphLine, &LineWidth, &CurrentWordLength, &CurrentLine, &MaxFontAscender, &MaxFontDescender, &NewLines, &CurrentWord, &CurrentAdvances, &bIsRTL]
+	    		(const FGlyphEntry& CurrentGlyph, int32 GlyphIndex, int32 IterationCount)
+	    		{
+	    			const bool bVisibleGlyph = IsValidGlyph(CurrentGlyph);
+					const int32 CharIndex = bIsRTL ? Kernings.Num() - 1 - CharacterIndex : CharacterIndex;
+    				const float CurrentKerning = bVisibleGlyph ? Kernings[CharIndex] : 0.f;
+    				const float GlyphAdv = GlyphLine->GetWidthAdvance(GlyphIndex, Tracking + CurrentKerning, WordSpacing);
+
+		            // If we're at the end the line or at whitespace
+		            const bool bWordBreak = !bVisibleGlyph || IterationCount == GlyphCount - 1;
+
+    				if (bWrap                               // when we're wrapping
+						&& bWordBreak                       // and at a word break
+						&& LineWidth > MaxWidth             // and the current line is longer than the max
+						&& CurrentWordLength != LineWidth)  // and the line is not just a single word that we can't break
+    				{
+    					CurrentLine->Width = LineWidth - CurrentWordLength;
+    					CurrentLine->MaxFontHeight = MaxFontAscender + MaxFontDescender;
+    					CurrentLine->MaxFontAscender = MaxFontAscender;
+    					CurrentLine->MaxFontDescender = MaxFontDescender;
+
+    					// Reset
+    					MaxFontAscender = 0.0f;
+    					MaxFontDescender = 0.0f;
+    					
+    					if (bIsRTL)
+    					{
+    						Algo::Reverse(CurrentLine->Glyphs);
+    						Algo::Reverse(CurrentLine->GlyphAdvances);
+    					}
+
+    					CurrentLine = &NewLines.AddDefaulted_GetRef();
+    					CurrentLine->TextDirection = GlyphLine->TextDirection;
+    					LineWidth = CurrentWordLength;
+    				}
+
+					LineWidth += GlyphAdv;
+					CurrentWordLength += GlyphAdv;
+
+					if (bVisibleGlyph)
+					{
+		#if WITH_FREETYPE
+						const FT_Face Face = FontFaces[CharIndex]->GetFace();
+
+						const float Ascender = FMath::Max(0.0f, Face->size->metrics.ascender * UE::Text3D::Metrics::Convert26Dot6ToPixel);
+						const float Descender = FMath::Max(0.0f, -Face->size->metrics.descender * UE::Text3D::Metrics::Convert26Dot6ToPixel);
+
+						MaxFontAscender  = FMath::Max(MaxFontAscender,  Ascender);
+						MaxFontDescender = FMath::Max(MaxFontDescender, Descender);
+		#endif
+					}
+
+					CurrentWord.Add(CurrentGlyph);
 					CurrentAdvances.Add(GlyphAdv);
-				}
-
-	            if (bWordBreak)
-	            {
-        			if (bIsRTL)
-        			{
-        				CurrentLine->Glyphs.Insert(CurrentWord, 0);
-        				CurrentLine->GlyphAdvances.Insert(CurrentAdvances, 0);
-        			}
-		            else
+
+		            if (bWordBreak)
 		            {
-            			CurrentLine->Glyphs.Append(CurrentWord);
-            			CurrentLine->GlyphAdvances.Append(CurrentAdvances);
+	            		CurrentLine->Glyphs.Append(CurrentWord);
+	            		CurrentLine->GlyphAdvances.Append(CurrentAdvances);
+		                CurrentWordLength = 0.0f;
+		                CurrentWord.Reset();
+		                CurrentAdvances.Reset();
 		            }
 
-	                CurrentWordLength = 0.0f;
-	                CurrentWord.Reset();
-	                CurrentAdvances.Reset();
-	            }
-
-    			if (bVisibleGlyph)
-    			{
-    				CharacterIndex++;
-    			}
-	        }
+    				if (bVisibleGlyph)
+    				{
+    					CharacterIndex++;
+    				}
+	    		}
+	    	);
 
 	        CurrentLine->Width = LineWidth;
     		CurrentLine->MaxFontHeight = MaxFontAscender + MaxFontDescender;
     		CurrentLine->MaxFontAscender = MaxFontAscender;
     		CurrentLine->MaxFontDescender = MaxFontDescender;
+	    	
+	    	if (bIsRTL)
+	    	{
+	    		Algo::Reverse(CurrentLine->Glyphs);
+	    		Algo::Reverse(CurrentLine->GlyphAdvances);
+	    	}
 	    }
 
 	    Lines = NewLines;
