Text3D: Fix nullptr crash when undoing style overrides

diff --git a/Engine/Plugins/VirtualProduction/Text3D/Source/Text3D/Private/Extensions/Text3DDefaultStyleExtension.cpp b/Engine/Plugins/VirtualProduction/Text3D/Source/Text3D/Private/Extensions/Text3DDefaultStyleExtension.cpp
--- a/Engine/Plugins/VirtualProduction/Text3D/Source/Text3D/Private/Extensions/Text3DDefaultStyleExtension.cpp
+++ b/Engine/Plugins/VirtualProduction/Text3D/Source/Text3D/Private/Extensions/Text3DDefaultStyleExtension.cpp
@@ -70,27 +70,24 @@ TSharedPtr<FSlateStyleSet> UText3DDefaultStyleExtension::GetCustomStyles() const
 	return TextStyleSet;
 }
 
-UText3DStyleBase* UText3DDefaultStyleExtension::GetStyle(FName InStyle) const
+UText3DStyleBase* UText3DDefaultStyleExtension::GetStyle(FName InStyleName) const
 {
+	auto FindStyleFunctor = [InStyleName](UText3DStyleBase* InStyle)->bool
+		{
+			return InStyle && InStyle->GetStyleName() == InStyleName;
+		};
+
 	// Look for overrides first
-	for (UText3DStyleBase* FormatStyle : Styles)
+	if (const TObjectPtr<UText3DStyleBase>* FormatStyle = Styles.FindByPredicate(FindStyleFunctor))
 	{
-		const FName StyleName = FormatStyle->GetStyleName();
-		if (StyleName.IsEqual(InStyle))
-		{
-			return FormatStyle;
-		}
+		return *FormatStyle;
 	}
 
 	if (StyleSet)
 	{
-		for (UText3DStyleBase* FormatStyle : StyleSet->GetStyles())
+		if (const TObjectPtr<UText3DStyleBase>* FormatStyle = StyleSet->GetStyles().FindByPredicate(FindStyleFunctor))
 		{
-			const FName StyleName = FormatStyle->GetStyleName();
-			if (StyleName.IsEqual(InStyle))
-			{
-				return FormatStyle;
-			}
+			return *FormatStyle;
 		}
 	}
 
