MotionDesign : Text3D
- Fixed contours not processed causing some inner geometry to not be removed

diff --git a/Engine/Plugins/VirtualProduction/Text3D/Source/Text3D/Private/GeometryBuilders/Text3DGlyphLoader.cpp b/Engine/Plugins/VirtualProduction/Text3D/Source/Text3D/Private/GeometryBuilders/Text3DGlyphLoader.cpp
--- a/Engine/Plugins/VirtualProduction/Text3D/Source/Text3D/Private/GeometryBuilders/Text3DGlyphLoader.cpp
+++ b/Engine/Plugins/VirtualProduction/Text3D/Source/Text3D/Private/GeometryBuilders/Text3DGlyphLoader.cpp
@@ -589,11 +589,6 @@ bool FText3DGlyphLoader::MergedWithNext(const TSharedPtr<FPolygon2f> Contour, co
 TSharedPtr<FPolygon2f> FText3DGlyphLoader::RemoveBadPoints(const TSharedPtr<FPolygon2f> Contour) const
 {
 	int32 VertexCount = Contour->VertexCount();
-	if (VertexCount < 4)
-	{
-		return nullptr;
-	}
-
 	int32 CheckVertexIndex = 0;
 	int32 ConsecutiveCleanVertices = 0;
 
