Text3D: fix issue where 'pure' blueprint actors with Text3D components sometimes do not show any component on load.
When constructing a blueprint actor, native components will first be iterated and find the first (if any) component that does not have an attach parent and is outered to AActor, and deem that as the root component. If the Actor has a valid Root Component set this way, it will skip the 'Default Scene Root' node (the scene component that gets created for purely blueprint actors).
The issue is caused when TextRoot gets renamed to a different outer (the actor) when it initially was outered to the renderer object, which caused this component to be considered for root and skip the default scene root component.

diff --git a/Engine/Plugins/VirtualProduction/Text3D/Source/Text3D/Private/Renderers/StaticMeshes/Text3DStaticMeshesRenderer.cpp b/Engine/Plugins/VirtualProduction/Text3D/Source/Text3D/Private/Renderers/StaticMeshes/Text3DStaticMeshesRenderer.cpp
--- a/Engine/Plugins/VirtualProduction/Text3D/Source/Text3D/Private/Renderers/StaticMeshes/Text3DStaticMeshesRenderer.cpp
+++ b/Engine/Plugins/VirtualProduction/Text3D/Source/Text3D/Private/Renderers/StaticMeshes/Text3DStaticMeshesRenderer.cpp
@@ -41,9 +41,9 @@ void UText3DStaticMeshesRenderer::OnCreate()
 		TextRoot->DetachFromComponent(DetachRule);
 		TextRoot->AttachToComponent(TextComponent, AttachRule);
 
-		if (TextRoot->GetOwner() != Owner)
+		if (TextRoot->GetOuter() != this)
 		{
-			TextRoot->Rename(nullptr, Owner, REN_DontCreateRedirectors | REN_NonTransactional | REN_DoNotDirty);
+			TextRoot->Rename(nullptr, this, REN_DontCreateRedirectors | REN_NonTransactional | REN_DoNotDirty);
 		}
 	}
 
