MotionDesign : Text3D
- Fixed similar crashes reported by CrashReport tool
- Added additional safety measures and warning logs

diff --git a/Engine/Plugins/VirtualProduction/Text3D/Source/Text3D/Private/Renderers/StaticMeshes/Text3DStaticMeshesRenderer.cpp b/Engine/Plugins/VirtualProduction/Text3D/Source/Text3D/Private/Renderers/StaticMeshes/Text3DStaticMeshesRenderer.cpp
--- a/Engine/Plugins/VirtualProduction/Text3D/Source/Text3D/Private/Renderers/StaticMeshes/Text3DStaticMeshesRenderer.cpp
+++ b/Engine/Plugins/VirtualProduction/Text3D/Source/Text3D/Private/Renderers/StaticMeshes/Text3DStaticMeshesRenderer.cpp
@@ -54,6 +54,12 @@ void UText3DStaticMeshesRenderer::OnUpdate(const UE::Text3D::Renderer::FUpdatePa
 {
 	const UText3DComponent* TextComponent = GetText3DComponent();
 
+	if (!TextRoot)
+	{
+		UE_LOG(LogText3D, Warning, TEXT("Invalid root component for %s renderer, cannot proceed!"), *GetFriendlyName().ToString())
+		return;
+	}
+
 	if (InParameters.CurrentFlag == EText3DRendererFlags::Geometry)
 	{
 		AllocateComponents(TextComponent->GetCharacterCount());
@@ -62,13 +68,20 @@ void UText3DStaticMeshesRenderer::OnUpdate(const UE::Text3D::Renderer::FUpdatePa
 		{
 			if (CharacterMeshes.IsValidIndex(InIndex))
 			{
-				if (const FText3DCachedMesh* CachedMesh = InCharacter->GetGlyphMesh())
+				if (UStaticMeshComponent* CharacterMeshComponent = CharacterMeshes[InIndex])
 				{
-					CharacterMeshes[InIndex]->SetStaticMesh(CachedMesh->StaticMesh);
+					if (const FText3DCachedMesh* CachedMesh = InCharacter->GetGlyphMesh())
+					{
+						CharacterMeshComponent->SetStaticMesh(CachedMesh->StaticMesh);
+					}
+					else
+					{
+						CharacterMeshComponent->SetStaticMesh(nullptr);
+					}
 				}
 				else
 				{
-					CharacterMeshes[InIndex]->SetStaticMesh(nullptr);
+					UE_LOG(LogText3D, Warning, TEXT("Invalid component object for character %i"), InIndex)
 				}
 			}
 			else
@@ -92,7 +105,14 @@ void UText3DStaticMeshesRenderer::OnUpdate(const UE::Text3D::Renderer::FUpdatePa
 
 			if (CharacterMeshes.IsValidIndex(InIndex))
 			{
-				CharacterMeshes[InIndex]->SetRelativeTransform(CharacterTransform);
+				if (UStaticMeshComponent* CharacterMeshComponent = CharacterMeshes[InIndex])
+				{
+					CharacterMeshComponent->SetRelativeTransform(CharacterTransform);
+				}
+				else
+				{
+					UE_LOG(LogText3D, Warning, TEXT("Invalid component object for character %i"), InIndex)
+				}
 			}
 			else
 			{
@@ -112,30 +132,36 @@ void UText3DStaticMeshesRenderer::OnUpdate(const UE::Text3D::Renderer::FUpdatePa
 		{
 			if (CharacterMeshes.IsValidIndex(InIndex))
 			{
-				UStaticMeshComponent* StaticMeshComponent = CharacterMeshes[InIndex];
-				for (int32 GroupIndex = 0; GroupIndex < static_cast<int32>(EText3DGroupType::TypeCount); GroupIndex++)
+				if (UStaticMeshComponent* CharacterMeshComponent = CharacterMeshes[InIndex])
 				{
-					const int32 MaterialIndex = StaticMeshComponent->GetMaterialIndex(SlotNames[GroupIndex]);
-
-					if (MaterialIndex == INDEX_NONE)
+					for (int32 GroupIndex = 0; GroupIndex < static_cast<int32>(EText3DGroupType::TypeCount); GroupIndex++)
 					{
-						continue;
-					}
+						const int32 MaterialIndex = CharacterMeshComponent->GetMaterialIndex(SlotNames[GroupIndex]);
 
-					const EText3DGroupType GroupType = static_cast<EText3DGroupType>(GroupIndex);
-					const FName StyleTag = InCharacter->GetStyleTag();
+						if (MaterialIndex == INDEX_NONE)
+						{
+							continue;
+						}
 
-					FMaterialParameters Parameters;
-					Parameters.Group = GroupType;
-					Parameters.Tag = StyleTag;
+						const EText3DGroupType GroupType = static_cast<EText3DGroupType>(GroupIndex);
+						const FName StyleTag = InCharacter->GetStyleTag();
 
-					UMaterialInterface* Material = MaterialExtension->GetMaterial(Parameters);
+						FMaterialParameters Parameters;
+						Parameters.Group = GroupType;
+						Parameters.Tag = StyleTag;
 
-					if (Material != StaticMeshComponent->GetMaterial(MaterialIndex))
-					{
-						StaticMeshComponent->SetMaterial(MaterialIndex, Material);
+						UMaterialInterface* Material = MaterialExtension->GetMaterial(Parameters);
+
+						if (Material != CharacterMeshComponent->GetMaterial(MaterialIndex))
+						{
+							CharacterMeshComponent->SetMaterial(MaterialIndex, Material);
+						}
 					}
 				}
+				else
+				{
+					UE_LOG(LogText3D, Warning, TEXT("Invalid component object for character %i"), InIndex)
+				}
 			}
 			else
 			{
@@ -151,14 +177,20 @@ void UText3DStaticMeshesRenderer::OnUpdate(const UE::Text3D::Renderer::FUpdatePa
 		{
 			if (CharacterMeshes.IsValidIndex(InIndex))
 			{
-				UStaticMeshComponent* StaticMeshComponent = CharacterMeshes[InIndex];
-				StaticMeshComponent->SetHiddenInGame(TextComponent->bHiddenInGame);
-				StaticMeshComponent->SetVisibility(TextComponent->GetVisibleFlag() && InCharacter->GetVisibility());
-				StaticMeshComponent->SetCastShadow(RenderingExtension->GetTextCastShadow());
-				StaticMeshComponent->SetCastHiddenShadow(RenderingExtension->GetTextCastHiddenShadow());
-				StaticMeshComponent->SetAffectDynamicIndirectLighting(RenderingExtension->GetTextAffectDynamicIndirectLighting());
-				StaticMeshComponent->SetAffectIndirectLightingWhileHidden(RenderingExtension->GetTextAffectIndirectLightingWhileHidden());
-				StaticMeshComponent->SetHoldout(RenderingExtension->GetTextHoldout());
+				if (UStaticMeshComponent* CharacterMeshComponent = CharacterMeshes[InIndex])
+				{
+					CharacterMeshComponent->SetHiddenInGame(TextComponent->bHiddenInGame);
+					CharacterMeshComponent->SetVisibility(TextComponent->GetVisibleFlag() && InCharacter->GetVisibility());
+					CharacterMeshComponent->SetCastShadow(RenderingExtension->GetTextCastShadow());
+					CharacterMeshComponent->SetCastHiddenShadow(RenderingExtension->GetTextCastHiddenShadow());
+					CharacterMeshComponent->SetAffectDynamicIndirectLighting(RenderingExtension->GetTextAffectDynamicIndirectLighting());
+					CharacterMeshComponent->SetAffectIndirectLightingWhileHidden(RenderingExtension->GetTextAffectIndirectLightingWhileHidden());
+					CharacterMeshComponent->SetHoldout(RenderingExtension->GetTextHoldout());
+				}
+				else
+				{
+					UE_LOG(LogText3D, Warning, TEXT("Invalid component object for character %i"), InIndex)
+				}
 			}
 			else
 			{
@@ -214,6 +246,12 @@ FBox UText3DStaticMeshesRenderer::OnCalculateBounds() const
 
 void UText3DStaticMeshesRenderer::AllocateComponents(int32 InCount)
 {
+	CharacterMeshes.RemoveAll([](const UStaticMeshComponent* InComponent)
+		{
+			return !InComponent;
+		}
+	);
+
 	if (CharacterMeshes.Num() == InCount)
 	{
 		return;
@@ -255,7 +293,7 @@ void UText3DStaticMeshesRenderer::AllocateComponents(int32 InCount)
 		for (int32 CharacterIndex = 0; CharacterIndex < FMath::Abs(RemainingCharacterCount); ++CharacterIndex)
 		{
 			UStaticMeshComponent* MeshComponent = nullptr;
-			if (!CharacterMeshesPool.IsEmpty())
+			while (!CharacterMeshesPool.IsEmpty() && !MeshComponent)
 			{
 				MeshComponent = CharacterMeshesPool.Pop();
 			}
diff --git a/Engine/Plugins/VirtualProduction/Text3D/Source/Text3D/Private/Renderers/Text3DRendererBase.cpp b/Engine/Plugins/VirtualProduction/Text3D/Source/Text3D/Private/Renderers/Text3DRendererBase.cpp
--- a/Engine/Plugins/VirtualProduction/Text3D/Source/Text3D/Private/Renderers/Text3DRendererBase.cpp
+++ b/Engine/Plugins/VirtualProduction/Text3D/Source/Text3D/Private/Renderers/Text3DRendererBase.cpp
@@ -106,24 +106,76 @@ void UText3DRendererBase::Update(EText3DRendererFlags InFlags)
 	{
 		return;
 	}
+	
+	// Let's make sure text extension subobjects are valid and usable !
+	UText3DTokenExtensionBase* TokenExtension = Text3DComponent->GetTokenExtension();
+	if (!TokenExtension)
+	{
+		UE_LOG(LogText3D, Warning, TEXT("Invalid Text3D Token extension, cannot proceed !"))
+		return;
+	}
 
-	TArray<UText3DExtensionBase*, TInlineAllocator<7>> Extensions
+	UText3DStyleExtensionBase* StyleExtension = Text3DComponent->GetStyleExtension();
+	if (!StyleExtension)
 	{
-		Text3DComponent->GetTokenExtension(),
-		Text3DComponent->GetStyleExtension(),
-		Text3DComponent->GetCharacterExtension(),
-		Text3DComponent->GetGeometryExtension(),
-		Text3DComponent->GetLayoutExtension(),
-		Text3DComponent->GetMaterialExtension(),
-		Text3DComponent->GetRenderingExtension()
-	};
+		UE_LOG(LogText3D, Warning, TEXT("Invalid Text3D Style extension, cannot proceed !"))
+		return;
+	}
+
+	UText3DCharacterExtensionBase* CharacterExtension = Text3DComponent->GetCharacterExtension();
+	if (!CharacterExtension)
+	{
+		UE_LOG(LogText3D, Warning, TEXT("Invalid Text3D Character extension, cannot proceed !"))
+		return;
+	}
 
-	Extensions.Append(Text3DComponent->GetLayoutEffects());
+	UText3DGeometryExtensionBase* GeometryExtension = Text3DComponent->GetGeometryExtension();
+	if (!GeometryExtension)
+	{
+		UE_LOG(LogText3D, Warning, TEXT("Invalid Text3D Geometry extension, cannot proceed !"))
+		return;
+	}
 
-	Extensions.RemoveAll([](const UText3DExtensionBase* InExtension)
+	UText3DLayoutExtensionBase* LayoutExtension = Text3DComponent->GetLayoutExtension();
+	if (!LayoutExtension)
 	{
-		return !IsValid(InExtension);
-	});
+		UE_LOG(LogText3D, Warning, TEXT("Invalid Text3D Layout extension, cannot proceed !"))
+		return;
+	}
+
+	UText3DMaterialExtensionBase* MaterialExtension = Text3DComponent->GetMaterialExtension();
+	if (!MaterialExtension)
+	{
+		UE_LOG(LogText3D, Warning, TEXT("Invalid Text3D Material extension, cannot proceed !"))
+		return;
+	}
+
+	UText3DRenderingExtensionBase* RenderingExtension = Text3DComponent->GetRenderingExtension();
+	if (!RenderingExtension)
+	{
+		UE_LOG(LogText3D, Warning, TEXT("Invalid Text3D Rendering extension, cannot proceed !"))
+		return;
+	}
+
+	TArray<UText3DExtensionBase*, TInlineAllocator<8>> Extensions
+	{
+		TokenExtension,
+		StyleExtension,
+		CharacterExtension,
+		GeometryExtension,
+		LayoutExtension,
+		MaterialExtension,
+		RenderingExtension
+	};
+
+	// Add optional extensions for layout effects
+	for (UText3DLayoutEffectBase* LayoutEffect : Text3DComponent->GetLayoutEffects())
+	{
+		if (LayoutEffect)
+		{
+			Extensions.Add(LayoutEffect);
+		}
+	}
 
 	// Sort in reverse order since we iterate in reverse order
 	Extensions.StableSort([](const UText3DExtensionBase& InExtensionA, const UText3DExtensionBase& InExtensionB)->bool
