MotionDesign : Text3D
- Fixed sequencer spawnable breaks rich text styles

diff --git a/Engine/Plugins/VirtualProduction/Text3D/Source/Text3D/Private/Extensions/Text3DDefaultCharacterExtension.cpp b/Engine/Plugins/VirtualProduction/Text3D/Source/Text3D/Private/Extensions/Text3DDefaultCharacterExtension.cpp
--- a/Engine/Plugins/VirtualProduction/Text3D/Source/Text3D/Private/Extensions/Text3DDefaultCharacterExtension.cpp
+++ b/Engine/Plugins/VirtualProduction/Text3D/Source/Text3D/Private/Extensions/Text3DDefaultCharacterExtension.cpp
@@ -35,6 +35,12 @@ void UText3DDefaultCharacterExtension::AllocateCharacters(uint16 InCount)
 
 void UText3DDefaultCharacterExtension::AllocateTextCharacters(uint16 InCharacterCount)
 {
+	TextCharacters.RemoveAll([](const UText3DCharacterBase* InCharacter)
+		{
+			return !InCharacter;
+		}
+	);
+
 	if (TextCharacters.Num() == InCharacterCount)
 	{
 		return;
@@ -61,7 +67,7 @@ void UText3DDefaultCharacterExtension::AllocateTextCharacters(uint16 InCharacter
 		{
 			UText3DCharacterBase* TextCharacter = nullptr;
 
-			if (!TextCharactersPool.IsEmpty())
+			while (!TextCharactersPool.IsEmpty() && !TextCharacter)
 			{
 				TextCharacter = TextCharactersPool.Pop();
 			}
diff --git a/Engine/Plugins/VirtualProduction/Text3D/Source/Text3D/Private/Extensions/Text3DDefaultMaterialExtension.cpp b/Engine/Plugins/VirtualProduction/Text3D/Source/Text3D/Private/Extensions/Text3DDefaultMaterialExtension.cpp
--- a/Engine/Plugins/VirtualProduction/Text3D/Source/Text3D/Private/Extensions/Text3DDefaultMaterialExtension.cpp
+++ b/Engine/Plugins/VirtualProduction/Text3D/Source/Text3D/Private/Extensions/Text3DDefaultMaterialExtension.cpp
@@ -390,6 +390,10 @@ EText3DExtensionResult UText3DDefaultMaterialExtension::PreRendererUpdate(const
 					if (UMaterialInstanceDynamic* DynamicMaterial = FindOrAdd(EText3DGroupType::Front, MaterialOverride.Tag))
 					{
 						const FLinearColor SolidColor = StyleRule->GetFrontColor();
+#if WITH_EDITOR
+						// Workaround to trigger an update of the material (eg: after a save operation, for spawnable)
+						DynamicMaterial->ClearParameterValues();
+#endif
 						DynamicMaterial->SetVectorParameterValue(UText3DProjectSettings::FMaterialParameters::SolidColor, SolidColor);
 						DynamicMaterial->SetScalarParameterValue(UText3DProjectSettings::FMaterialParameters::Mode, 0);
 						DynamicMaterial->SetScalarParameterValue(UText3DProjectSettings::FMaterialParameters::Opacity, Opacity);
