Automate '3D Text - 3D Text's Horizontal and Vertical alignment can be changed'

#rnx

diff --git a/Engine/Plugins/VirtualProduction/Text3D/Source/Text3DEditor/Private/Widgets/SText3DEditorHorizontalAlignment.h b/Engine/Plugins/VirtualProduction/Text3D/Source/Text3DEditor/Private/Widgets/SText3DEditorHorizontalAlignment.h
--- a/Engine/Plugins/VirtualProduction/Text3D/Source/Text3DEditor/Private/Widgets/SText3DEditorHorizontalAlignment.h
+++ b/Engine/Plugins/VirtualProduction/Text3D/Source/Text3DEditor/Private/Widgets/SText3DEditorHorizontalAlignment.h
@@ -27,10 +27,12 @@ public:
 private:
 	TSharedRef<SButton> MakeAlignmentButton(EText3DHorizontalTextAlignment InAlignment, FName InBrushName, const FText& InTooltip);
 	FSlateColor GetButtonColorAndOpacity(EText3DHorizontalTextAlignment InAlignment) const;
-	FReply OnAlignmentButtonClicked(EText3DHorizontalTextAlignment InAlignment);
+	TEXT3DEDITOR_API FReply OnAlignmentButtonClicked(EText3DHorizontalTextAlignment InAlignment);
 	EText3DHorizontalTextAlignment GetPropertyAlignment() const;
 	void OnPropertyChanged();
 
 	TSharedPtr<IPropertyHandle> PropertyHandle;
 	FOnHorizontalAlignmentChanged AlignmentChangedDelegate;
+
+	friend class UText3DEditorAccessHelper;
 };
diff --git a/Engine/Plugins/VirtualProduction/Text3D/Source/Text3DEditor/Private/Widgets/SText3DEditorVerticalAlignment.h b/Engine/Plugins/VirtualProduction/Text3D/Source/Text3DEditor/Private/Widgets/SText3DEditorVerticalAlignment.h
--- a/Engine/Plugins/VirtualProduction/Text3D/Source/Text3DEditor/Private/Widgets/SText3DEditorVerticalAlignment.h
+++ b/Engine/Plugins/VirtualProduction/Text3D/Source/Text3DEditor/Private/Widgets/SText3DEditorVerticalAlignment.h
@@ -26,10 +26,12 @@ public:
 private:
 	TSharedRef<SButton> MakeAlignmentButton(EText3DVerticalTextAlignment InAlignment, FName InBrushName, const FText& InTooltip);
 	FSlateColor GetButtonColorAndOpacity(EText3DVerticalTextAlignment InAlignment) const;
-	FReply OnAlignmentButtonClicked(EText3DVerticalTextAlignment InAlignment);
+	TEXT3DEDITOR_API FReply OnAlignmentButtonClicked(EText3DVerticalTextAlignment InAlignment);
 	EText3DVerticalTextAlignment GetPropertyAlignment() const;
 	void OnPropertyChanged();
 
 	TSharedPtr<IPropertyHandle> PropertyHandle;
 	FOnVerticalAlignmentChanged AlignmentChangedDelegate;
+
+	friend class UText3DEditorAccessHelper;
 };
\ No newline at end of file
