Automate '3D Text - Change Text Extrude parameter and validate its execution'

#rnx

diff --git a/Engine/Plugins/VirtualProduction/Text3D/Source/Text3D/Private/GeometryBuilders/Text3DGlyphContourNode.h b/Engine/Plugins/VirtualProduction/Text3D/Source/Text3D/Private/GeometryBuilders/Text3DGlyphContourNode.h
--- a/Engine/Plugins/VirtualProduction/Text3D/Source/Text3D/Private/GeometryBuilders/Text3DGlyphContourNode.h
+++ b/Engine/Plugins/VirtualProduction/Text3D/Source/Text3D/Private/GeometryBuilders/Text3DGlyphContourNode.h
@@ -3,7 +3,7 @@
 #pragma once
 
 #include "Containers/Array.h"
-#include "Logs/Text3DLogs.h"
+#include "Text3D/Private/Logs/Text3DLogs.h"
 #include "Polygon2.h"
 #include "Templates/SharedPointer.h"
 
diff --git a/Engine/Plugins/VirtualProduction/Text3D/Source/Text3D/Private/Subsystems/Text3DEngineSubsystem.h b/Engine/Plugins/VirtualProduction/Text3D/Source/Text3D/Private/Subsystems/Text3DEngineSubsystem.h
--- a/Engine/Plugins/VirtualProduction/Text3D/Source/Text3D/Private/Subsystems/Text3DEngineSubsystem.h
+++ b/Engine/Plugins/VirtualProduction/Text3D/Source/Text3D/Private/Subsystems/Text3DEngineSubsystem.h
@@ -3,7 +3,7 @@
 #pragma once
 
 #include "Containers/Ticker.h"
-#include "GeometryBuilders/Text3DGlyphContourNode.h"
+#include "Text3D/Private/GeometryBuilders/Text3DGlyphContourNode.h"
 #include "Subsystems/EngineSubsystem.h"
 #include "Text3DTypes.h"
 
diff --git a/Engine/Plugins/VirtualProduction/Text3D/Source/Text3D/Public/Extensions/Text3DDefaultGeometryExtension.h b/Engine/Plugins/VirtualProduction/Text3D/Source/Text3D/Public/Extensions/Text3DDefaultGeometryExtension.h
--- a/Engine/Plugins/VirtualProduction/Text3D/Source/Text3D/Public/Extensions/Text3DDefaultGeometryExtension.h
+++ b/Engine/Plugins/VirtualProduction/Text3D/Source/Text3D/Public/Extensions/Text3DDefaultGeometryExtension.h
@@ -3,7 +3,7 @@
 #pragma once
 
 #include "Text3DGeometryExtensionBase.h"
-#include "Subsystems/Text3DEngineSubsystem.h"
+#include "Text3D/Private/Subsystems/Text3DEngineSubsystem.h"
 #include "Text3DDefaultGeometryExtension.generated.h"
 
 UCLASS(MinimalAPI)
