MotionDesign : Text3D
- Added layout transform stagger option to animate each character individually

diff --git a/Engine/Plugins/VirtualProduction/Text3D/Source/Text3D/Private/Extensions/Text3DLayoutTransformEffect.cpp b/Engine/Plugins/VirtualProduction/Text3D/Source/Text3D/Private/Extensions/Text3DLayoutTransformEffect.cpp
--- a/Engine/Plugins/VirtualProduction/Text3D/Source/Text3D/Private/Extensions/Text3DLayoutTransformEffect.cpp
+++ b/Engine/Plugins/VirtualProduction/Text3D/Source/Text3D/Private/Extensions/Text3DLayoutTransformEffect.cpp
@@ -63,6 +63,17 @@ void UText3DLayoutTransformEffect::SetLocationEaseCurve(UCurveFloat* InEaseCurve
 	OnTransformOptionsChanged();
 }
 
+void UText3DLayoutTransformEffect::SetLocationStagger(bool bInStagger)
+{
+	if (bLocationStagger == bInStagger)
+	{
+		return;
+	}
+
+	bLocationStagger = bInStagger;
+	OnTransformOptionsChanged();
+}
+
 void UText3DLayoutTransformEffect::SetScaleEnabled(bool bEnabled)
 {
 	if (bScaleEnabled != bEnabled)
@@ -120,6 +131,17 @@ void UText3DLayoutTransformEffect::SetScaleEaseCurve(UCurveFloat* InEaseCurve)
 	OnTransformOptionsChanged();
 }
 
+void UText3DLayoutTransformEffect::SetScaleStagger(bool bInStagger)
+{
+	if (bScaleStagger == bInStagger)
+	{
+		return;
+	}
+
+	bScaleStagger = bInStagger;
+	OnTransformOptionsChanged();
+}
+
 void UText3DLayoutTransformEffect::SetRotationEnabled(bool bEnabled)
 {
 	if (bRotationEnabled != bEnabled)
@@ -177,6 +199,17 @@ void UText3DLayoutTransformEffect::SetRotationEaseCurve(UCurveFloat* InEaseCurve
 	OnTransformOptionsChanged();
 }
 
+void UText3DLayoutTransformEffect::SetRotationStagger(bool bInStagger)
+{
+	if (bRotationStagger == bInStagger)
+	{
+		return;
+	}
+
+	bRotationStagger = bInStagger;
+	OnTransformOptionsChanged();
+}
+
 UText3DLayoutTransformEffect::UText3DLayoutTransformEffect()
 {
 	OnTransformOptionsChanged();
@@ -248,21 +281,21 @@ void UText3DLayoutTransformEffect::ApplyEffect(uint32 InGlyphIndex, uint32 InGly
 
 	if (bLocationEnabled)
 	{
-		const float Effect = CalculateEffect(InGlyphIndex, InGlyphCount, LocationOrder, LocationProgress, LocationEaseCurve);
+		const float Effect = CalculateEffect(InGlyphIndex, InGlyphCount, LocationOrder, LocationProgress, LocationEaseCurve, bLocationStagger);
 		const FVector Location = LocationBegin + (LocationEnd - LocationBegin) * Effect;
 		CharacterTransform.AddToTranslation(Location);
 	}
 
 	if (bScaleEnabled)
 	{
-		const float Effect = CalculateEffect(InGlyphIndex, InGlyphCount, ScaleOrder, ScaleProgress, ScaleEaseCurve);
+		const float Effect = CalculateEffect(InGlyphIndex, InGlyphCount, ScaleOrder, ScaleProgress, ScaleEaseCurve, bScaleStagger);
 		const FVector Scale = ScaleBegin + (ScaleEnd - ScaleBegin) * Effect;
 		CharacterTransform.MultiplyScale3D(Scale.ComponentMax(FVector(UE_KINDA_SMALL_NUMBER)));
 	}
 
 	if (bRotationEnabled)
 	{
-		const float Effect = CalculateEffect(InGlyphIndex, InGlyphCount, RotationOrder, RotationProgress, RotationEaseCurve);
+		const float Effect = CalculateEffect(InGlyphIndex, InGlyphCount, RotationOrder, RotationProgress, RotationEaseCurve, bRotationStagger);
 		const FRotator Rotator = RotationBegin + (RotationEnd - RotationBegin) * Effect;
 		CharacterTransform.ConcatenateRotation(Rotator.Quaternion());
 	}
@@ -303,13 +336,14 @@ int32 UText3DLayoutTransformEffect::GetEffectPosition(int32 Index, int32 Total,
 	return Index;
 }
 
-float UText3DLayoutTransformEffect::CalculateEffect(int32 InIndex, int32 InTotal, EText3DCharacterEffectOrder InOrder, float InProgress, const UCurveFloat* InEaseCurve) const
+float UText3DLayoutTransformEffect::CalculateEffect(int32 InIndex, int32 InTotal, EText3DCharacterEffectOrder InOrder, float InProgress, const UCurveFloat* InEaseCurve, bool bInStagger) const
 {
 	const int32 Position = GetEffectPosition(InIndex, InTotal, InOrder);
 	const float NormalizedProgress = FMath::Clamp(InProgress * 0.01f, 0.0f, 1.0f);
 	const float StaggerAmount = 1.0f / FMath::Max(1, InTotal);
 	const float CharacterStart = Position * StaggerAmount;
-	const float CharacterProgress = FMath::Clamp((NormalizedProgress - CharacterStart) / (1.0f - CharacterStart), 0.0f, 1.0f);
+	const float StaggerProgress = bInStagger ? StaggerAmount : 1.0f - CharacterStart;
+	const float CharacterProgress = FMath::Clamp((NormalizedProgress - CharacterStart) / StaggerProgress, 0.0f, 1.0f);
 	float CharacterEased = CharacterProgress;
 
 	if (InEaseCurve)
diff --git a/Engine/Plugins/VirtualProduction/Text3D/Source/Text3D/Public/Extensions/Text3DLayoutTransformEffect.h b/Engine/Plugins/VirtualProduction/Text3D/Source/Text3D/Public/Extensions/Text3DLayoutTransformEffect.h
--- a/Engine/Plugins/VirtualProduction/Text3D/Source/Text3D/Public/Extensions/Text3DLayoutTransformEffect.h
+++ b/Engine/Plugins/VirtualProduction/Text3D/Source/Text3D/Public/Extensions/Text3DLayoutTransformEffect.h
@@ -63,6 +63,14 @@ public:
 		return LocationEaseCurve;
 	}
 
+	UFUNCTION(BlueprintCallable, Category = "Text3D|Transform")
+	TEXT3D_API void SetLocationStagger(bool bInStagger);
+
+	bool GetLocationStagger() const
+	{
+		return bLocationStagger;
+	}
+
 	UFUNCTION(BlueprintCallable, Category = "Text3D|Transform")
 	TEXT3D_API void SetRotationEnabled(bool bEnabled);
 
@@ -111,6 +119,14 @@ public:
 		return RotationEaseCurve;
 	}
 
+	UFUNCTION(BlueprintCallable, Category = "Text3D|Transform")
+	TEXT3D_API void SetRotationStagger(bool bInStagger);
+
+	bool GetRotationStagger() const
+	{
+		return bRotationStagger;
+	}
+
 	UFUNCTION(BlueprintCallable, Category = "Text3D|Transform")
 	TEXT3D_API void SetScaleEnabled(bool bEnabled);
 
@@ -159,6 +175,13 @@ public:
 		return ScaleEaseCurve;
 	}
 
+	UFUNCTION(BlueprintCallable, Category = "Text3D|Transform")
+	TEXT3D_API void SetScaleStagger(bool bInStagger);
+
+	bool GetScaleStagger() const
+	{
+		return bScaleStagger;
+	}
 
 protected:
 	UText3DLayoutTransformEffect();
@@ -177,7 +200,7 @@ protected:
 	void OnTransformOptionsChanged();
 
 	int32 GetEffectPosition(int32 Index, int32 Total, EText3DCharacterEffectOrder Order) const;
-	float CalculateEffect(int32 InIndex, int32 InTotal, EText3DCharacterEffectOrder InOrder, float InProgress, const UCurveFloat* InEaseCurve) const;
+	float CalculateEffect(int32 InIndex, int32 InTotal, EText3DCharacterEffectOrder InOrder, float InProgress, const UCurveFloat* InEaseCurve, bool bInStagger) const;
 
 	// Location
 
@@ -200,6 +223,9 @@ protected:
 	UPROPERTY(EditAnywhere, Setter, Getter, Category = "Transform", meta = (EditCondition = "bLocationEnabled", EditConditionHides))
 	TObjectPtr<UCurveFloat> LocationEaseCurve;
 
+	UPROPERTY(EditAnywhere, Getter="GetLocationStagger", Setter="SetLocationStagger", Category = "Transform", meta = (EditCondition = "bLocationEnabled", EditConditionHides))
+	bool bLocationStagger = false; 
+
 	// Rotate
 
 	UPROPERTY(EditAnywhere, Getter="GetRotationEnabled", Setter="SetRotationEnabled", Category = "Transform")
@@ -221,6 +247,9 @@ protected:
 	UPROPERTY(EditAnywhere, Setter, Getter, Category = "Transform", meta = (EditCondition = "bRotationEnabled", EditConditionHides))
 	TObjectPtr<UCurveFloat> RotationEaseCurve;
 
+	UPROPERTY(EditAnywhere, Getter="GetRotationStagger", Setter="SetRotationStagger", Category = "Transform", meta = (EditCondition = "bRotationEnabled", EditConditionHides))
+	bool bRotationStagger = false;
+
 	// Scale
 
 	UPROPERTY(EditAnywhere, Getter="GetScaleEnabled", Setter="SetScaleEnabled", Category = "Transform")
@@ -241,4 +270,7 @@ protected:
 	/** Provide a 0-1 ease curve, leaving this unset will result in linear ease */
 	UPROPERTY(EditAnywhere, Setter, Getter, Category = "Transform", meta = (EditCondition = "bScaleEnabled", EditConditionHides))
 	TObjectPtr<UCurveFloat> ScaleEaseCurve;
+
+	UPROPERTY(EditAnywhere, Getter="GetScaleStagger", Setter="SetScaleStagger", Category = "Transform", meta = (EditCondition = "bScaleEnabled", EditConditionHides))
+	bool bScaleStagger = false;
 };
