Motion Design: Changed Shape actor and Text3D to static draw since shapes and Text3D are mostly going to be unchanging mesh in most frames.

diff --git a/Engine/Plugins/VirtualProduction/Avalanche/Source/AvalancheShapes/Private/AvaShapeActor.cpp b/Engine/Plugins/VirtualProduction/Avalanche/Source/AvalancheShapes/Private/AvaShapeActor.cpp
--- a/Engine/Plugins/VirtualProduction/Avalanche/Source/AvalancheShapes/Private/AvaShapeActor.cpp
+++ b/Engine/Plugins/VirtualProduction/Avalanche/Source/AvalancheShapes/Private/AvaShapeActor.cpp
@@ -23,6 +23,7 @@ AAvaShapeActor::AAvaShapeActor()
 	ShapeMeshComponent->SetGenerateOverlapEvents(true);
 	ShapeMeshComponent->SetCanEverAffectNavigation(false);
 	ShapeMeshComponent->SetComplexAsSimpleCollisionEnabled(true);
+	ShapeMeshComponent->SetMeshDrawPath(EDynamicMeshDrawPath::StaticDraw);
 
 	SetRootComponent(ShapeMeshComponent);
 
diff --git a/Engine/Plugins/VirtualProduction/Text3D/Source/Text3D/Private/Renderers/DynamicMesh/Text3DDynamicMeshRenderer.cpp b/Engine/Plugins/VirtualProduction/Text3D/Source/Text3D/Private/Renderers/DynamicMesh/Text3DDynamicMeshRenderer.cpp
--- a/Engine/Plugins/VirtualProduction/Text3D/Source/Text3D/Private/Renderers/DynamicMesh/Text3DDynamicMeshRenderer.cpp
+++ b/Engine/Plugins/VirtualProduction/Text3D/Source/Text3D/Private/Renderers/DynamicMesh/Text3DDynamicMeshRenderer.cpp
@@ -63,6 +63,7 @@ void UText3DDynamicMeshRenderer::OnCreate()
 		DynamicMeshComponent->SetComplexAsSimpleCollisionEnabled(true);
 		DynamicMeshComponent->SetDeferredCollisionUpdatesEnabled(true);
 		DynamicMeshComponent->SetCollisionProfileName(UCollisionProfile::BlockAllDynamic_ProfileName);
+		DynamicMeshComponent->SetMeshDrawPath(EDynamicMeshDrawPath::StaticDraw);
 
 		ClearMesh();
 	}
