Text3D: expose virtual functions and put them public to match the base

#Synced-CL 48692301

diff --git a/Engine/Plugins/VirtualProduction/Text3D/Source/Text3D/Public/Extensions/Text3DDefaultMaterialExtension.h b/Engine/Plugins/VirtualProduction/Text3D/Source/Text3D/Public/Extensions/Text3DDefaultMaterialExtension.h
--- a/Engine/Plugins/VirtualProduction/Text3D/Source/Text3D/Public/Extensions/Text3DDefaultMaterialExtension.h
+++ b/Engine/Plugins/VirtualProduction/Text3D/Source/Text3D/Public/Extensions/Text3DDefaultMaterialExtension.h
@@ -235,31 +235,31 @@ public:
 
 	TEXT3D_API void PreCacheMaterials();
 
-protected:
 	//~ Begin UObject
-	virtual void PostLoad() override;
-	virtual void PostInitProperties() override;
+	TEXT3D_API virtual void PostLoad() override;
+	TEXT3D_API virtual void PostInitProperties() override;
 #if WITH_EDITOR
-	virtual void PostEditChangeProperty(FPropertyChangedEvent& InEvent) override;
-	virtual void PostEditUndo() override;
+	TEXT3D_API virtual void PostEditChangeProperty(FPropertyChangedEvent& InEvent) override;
+	TEXT3D_API virtual void PostEditUndo() override;
 #endif
 	//~ End UObject
 
 	//~ Begin UText3DExtensionBase
-	virtual EText3DExtensionResult PreRendererUpdate(const UE::Text3D::Renderer::FUpdateParameters& InParameters) override;
-	virtual EText3DExtensionResult PostRendererUpdate(const UE::Text3D::Renderer::FUpdateParameters& InParameters) override;
+	TEXT3D_API virtual EText3DExtensionResult PreRendererUpdate(const UE::Text3D::Renderer::FUpdateParameters& InParameters) override;
+	TEXT3D_API virtual EText3DExtensionResult PostRendererUpdate(const UE::Text3D::Renderer::FUpdateParameters& InParameters) override;
 	//~ End UText3DExtensionBase
 
 	//~ Begin UText3DMaterialExtensionBase
-	virtual void SetMaterial(const UE::Text3D::Material::FMaterialParameters& InParameters, UMaterialInterface* InMaterial) override;
-	virtual UMaterialInterface* GetMaterial(const UE::Text3D::Material::FMaterialParameters& InParameters) const override;
-	virtual void RegisterMaterialOverride(FName InTag) override;
-	virtual void UnregisterMaterialOverride(FName InTag) override;
-	virtual void ForEachMaterial(TFunctionRef<bool(const UE::Text3D::Material::FMaterialParameters&, UMaterialInterface*)> InFunctor) const override;
-	virtual int32 GetMaterialCount() const override;
-	virtual void GetMaterialNames(TArray<FName>& OutNames) const override;
+	TEXT3D_API virtual void SetMaterial(const UE::Text3D::Material::FMaterialParameters& InParameters, UMaterialInterface* InMaterial) override;
+	TEXT3D_API virtual UMaterialInterface* GetMaterial(const UE::Text3D::Material::FMaterialParameters& InParameters) const override;
+	TEXT3D_API virtual void RegisterMaterialOverride(FName InTag) override;
+	TEXT3D_API virtual void UnregisterMaterialOverride(FName InTag) override;
+	TEXT3D_API virtual void ForEachMaterial(TFunctionRef<bool(const UE::Text3D::Material::FMaterialParameters&, UMaterialInterface*)> InFunctor) const override;
+	TEXT3D_API virtual int32 GetMaterialCount() const override;
+	TEXT3D_API virtual void GetMaterialNames(TArray<FName>& OutNames) const override;
 	//~ End UText3DMaterialExtensionBase
 
+protected:
 	UMaterialInstanceDynamic* FindOrAdd(EText3DGroupType InGroup, FName InTag = NAME_None);
 	UMaterialInstanceDynamic* Find(EText3DGroupType InGroup, FName InTag = NAME_None) const;
 	void SetVectorParameter(TConstArrayView<UMaterialInstanceDynamic*> InMaterials, FName InKey, FVector InValue);
