Text3D: added option to set text shape direction manually

#Synced-CL 49962738

diff --git a/Engine/Plugins/VirtualProduction/Text3D/Source/Text3D/Private/Extensions/Text3DDefaultLayoutExtension.cpp b/Engine/Plugins/VirtualProduction/Text3D/Source/Text3D/Private/Extensions/Text3DDefaultLayoutExtension.cpp
--- a/Engine/Plugins/VirtualProduction/Text3D/Source/Text3D/Private/Extensions/Text3DDefaultLayoutExtension.cpp
+++ b/Engine/Plugins/VirtualProduction/Text3D/Source/Text3D/Private/Extensions/Text3DDefaultLayoutExtension.cpp
@@ -173,6 +173,17 @@ void UText3DDefaultLayoutExtension::SetScaleProportionally(const bool Value)
 	OnLayoutOptionsChanged();
 }
 
+void UText3DDefaultLayoutExtension::SetTextShapingMethod(ETextShapingMethod InTextShapingMethod)
+{
+	if (TextShapingMethod == InTextShapingMethod)
+	{
+		return;
+	}
+
+	TextShapingMethod = InTextShapingMethod;
+	OnTextShapingMethodChanged();
+}
+
 float UText3DDefaultLayoutExtension::GetTextHeight() const
 {
 	using namespace UE::Text3D::Layout;
@@ -212,6 +223,7 @@ EText3DExtensionResult UText3DDefaultLayoutExtension::PreRendererUpdate(const UE
 	{
 		TextLayout = MakeShared<FText3DLayout>();
 	}
+	TextLayout->SetTextShapingMethod(TextShapingMethod);
 
 	if (!TextLayoutMarshaller.IsValid())
 	{
@@ -509,10 +521,15 @@ void UText3DDefaultLayoutExtension::PostEditChangeProperty(FPropertyChangedEvent
 		GET_MEMBER_NAME_CHECKED(UText3DDefaultLayoutExtension, bScaleProportionally),
 	};
 
-	if (LayoutPropertyNames.Contains(InEvent.GetMemberPropertyName()))
+	const FName ChangedPropertyName = InEvent.GetMemberPropertyName();
+	if (LayoutPropertyNames.Contains(ChangedPropertyName))
 	{
 		OnLayoutOptionsChanged();
 	}
+	else if (ChangedPropertyName == GET_MEMBER_NAME_CHECKED(UText3DDefaultLayoutExtension, TextShapingMethod))
+	{
+		OnTextShapingMethodChanged();
+	}
 }
 #endif
 
@@ -523,8 +540,7 @@ void UText3DDefaultLayoutExtension::OnLayoutOptionsChanged()
 	MaxWidth = FMath::Max(1.0f, MaxWidth);
 	MaxHeight = FMath::Max(1.0f, MaxHeight);
 
-	EText3DRendererFlags Flags = EText3DRendererFlags::Material;
-	EnumAddFlags(Flags, EText3DRendererFlags::Layout);
+	EText3DRendererFlags Flags = EText3DRendererFlags::Material | EText3DRendererFlags::Layout;
 
 	if (MaxWidthBehavior == EText3DMaxWidthHandling::WrapAndScale && GlyphText.IsValid())
 	{
@@ -532,7 +548,7 @@ void UText3DDefaultLayoutExtension::OnLayoutOptionsChanged()
 		{
 			if (ShapedLine->TextDirection == TextBiDi::ETextDirection::RightToLeft)
 			{
-				EnumAddFlags(Flags, EText3DRendererFlags::Geometry);
+				Flags |= EText3DRendererFlags::Geometry;
 				break;
 			}
 		}
@@ -540,3 +556,10 @@ void UText3DDefaultLayoutExtension::OnLayoutOptionsChanged()
 
 	RequestUpdate(Flags);
 }
+
+void UText3DDefaultLayoutExtension::OnTextShapingMethodChanged()
+{
+	// Changing text shaping method requires recomputing geometry too
+	constexpr EText3DRendererFlags Flags = EText3DRendererFlags::Material | EText3DRendererFlags::Layout | EText3DRendererFlags::Geometry;
+	RequestUpdate(Flags);
+}
diff --git a/Engine/Plugins/VirtualProduction/Text3D/Source/Text3D/Public/Extensions/Text3DDefaultLayoutExtension.h b/Engine/Plugins/VirtualProduction/Text3D/Source/Text3D/Public/Extensions/Text3DDefaultLayoutExtension.h
--- a/Engine/Plugins/VirtualProduction/Text3D/Source/Text3D/Public/Extensions/Text3DDefaultLayoutExtension.h
+++ b/Engine/Plugins/VirtualProduction/Text3D/Source/Text3D/Public/Extensions/Text3DDefaultLayoutExtension.h
@@ -3,6 +3,7 @@
 #pragma once
 
 #include "Extensions/Text3DLayoutExtensionBase.h"
+#include "Fonts/FontCache.h"
 #include "Text3DDefaultLayoutExtension.generated.h"
 
 class FRichTextLayoutMarshaller;
@@ -146,6 +147,17 @@ public:
 	UFUNCTION(BlueprintCallable, Category = "Text3D|Layout")
 	void SetScaleProportionally(const bool Value);
 
+	/** Get the current method of shaping text */
+	UFUNCTION(BlueprintPure, Category = "Text3D|Layout")
+	ETextShapingMethod GetTextShapingMethod() const
+	{
+		return TextShapingMethod;
+	}
+
+	/** Set the method to shape text */
+	UFUNCTION(BlueprintCallable, Category = "Text3D|Layout")
+	void SetTextShapingMethod(ETextShapingMethod InTextShapingMethod);
+
 protected:
 	//~ Begin UObject
 #if WITH_EDITOR
@@ -167,6 +179,7 @@ protected:
 
 	void CalculateTextScale();
 	void OnLayoutOptionsChanged();
+	void OnTextShapingMethodChanged();
 
 	/** Horizontal text alignment */
 	UPROPERTY(EditAnywhere, Getter, Setter, Category = "Layout", meta = (AllowPrivateAccess = "true", CustomAlignmentWidget))
@@ -212,6 +225,10 @@ protected:
 	UPROPERTY(EditAnywhere, Getter = "GetScalesProportionally", Setter = "SetScaleProportionally", Category = "Layout", meta = (AllowPrivateAccess = "true"))
 	bool bScaleProportionally = true;
 
+	/** The method to use for shaping the text */
+	UPROPERTY(EditAnywhere, Getter, Setter, Category = "Layout", meta = (AllowPrivateAccess = "true"), AdvancedDisplay)
+	ETextShapingMethod TextShapingMethod = ETextShapingMethod::Auto;
+
 private:
 	/** Additional scale to apply to the text. */
 	FVector TextScale = FVector::ZeroVector;
