Text3D: added missing includes

#Synced-CL 50556942

diff --git a/Engine/Plugins/VirtualProduction/Text3D/Source/Text3D/Public/Characters/Text3DCharacterBase.h b/Engine/Plugins/VirtualProduction/Text3D/Source/Text3D/Public/Characters/Text3DCharacterBase.h
--- a/Engine/Plugins/VirtualProduction/Text3D/Source/Text3D/Public/Characters/Text3DCharacterBase.h
+++ b/Engine/Plugins/VirtualProduction/Text3D/Source/Text3D/Public/Characters/Text3DCharacterBase.h
@@ -2,8 +2,12 @@
 
 #pragma once
 
+#include "Math/Rotator.h"
+#include "Math/Transform.h"
+#include "Math/Vector.h"
 #include "Misc/Optional.h"
 #include "Text3DTypes.h"
+#include "UObject/Object.h"
 #include "Text3DCharacterBase.generated.h"
 
 /** Holds data for a single character in Text3D */
diff --git a/Engine/Plugins/VirtualProduction/Text3D/Source/Text3D/Public/Extensions/Text3DExtensionBase.h b/Engine/Plugins/VirtualProduction/Text3D/Source/Text3D/Public/Extensions/Text3DExtensionBase.h
--- a/Engine/Plugins/VirtualProduction/Text3D/Source/Text3D/Public/Extensions/Text3DExtensionBase.h
+++ b/Engine/Plugins/VirtualProduction/Text3D/Source/Text3D/Public/Extensions/Text3DExtensionBase.h
@@ -3,6 +3,7 @@
 #pragma once
 
 #include "Text3DTypes.h"
+#include "UObject/Object.h"
 #include "Text3DExtensionBase.generated.h"
 
 class UText3DComponent;
diff --git a/Engine/Plugins/VirtualProduction/Text3D/Source/Text3D/Public/Renderers/Text3DRendererBase.h b/Engine/Plugins/VirtualProduction/Text3D/Source/Text3D/Public/Renderers/Text3DRendererBase.h
--- a/Engine/Plugins/VirtualProduction/Text3D/Source/Text3D/Public/Renderers/Text3DRendererBase.h
+++ b/Engine/Plugins/VirtualProduction/Text3D/Source/Text3D/Public/Renderers/Text3DRendererBase.h
@@ -4,6 +4,7 @@
 
 #include "Math/Box.h"
 #include "Text3DTypes.h"
+#include "UObject/Object.h"
 #include "Text3DRendererBase.generated.h"
 
 class UText3DComponent;
