Remove PrecacheMaterials from PostInit properties to prevent it from being executed in parallel on Async loading thread.

- It is already called from PostLoad() in that case

#Synced-CL 50453497

diff --git a/Engine/Plugins/VirtualProduction/Text3D/Source/Text3D/Private/Extensions/Text3DDefaultMaterialExtension.cpp b/Engine/Plugins/VirtualProduction/Text3D/Source/Text3D/Private/Extensions/Text3DDefaultMaterialExtension.cpp
--- a/Engine/Plugins/VirtualProduction/Text3D/Source/Text3D/Private/Extensions/Text3DDefaultMaterialExtension.cpp
+++ b/Engine/Plugins/VirtualProduction/Text3D/Source/Text3D/Private/Extensions/Text3DDefaultMaterialExtension.cpp
@@ -722,7 +722,11 @@ void UText3DDefaultMaterialExtension::PostInitProperties()
 	Super::PostInitProperties();
 
 	GroupDynamicMaterials.Empty();
-	PreCacheMaterials();
+	if (!IsInAsyncLoadingThread())
+	{
+		// This needs to be called on the game thread. In the case of loading it will get called in PostLoad()
+		PreCacheMaterials();
+	}
 }
 
 #if WITH_EDITOR
