Text3D: added missing include

#Synced-CL 50556942

diff --git a/Engine/Plugins/VirtualProduction/Text3D/Source/Text3D/Public/Text3DTypes.h b/Engine/Plugins/VirtualProduction/Text3D/Source/Text3D/Public/Text3DTypes.h
--- a/Engine/Plugins/VirtualProduction/Text3D/Source/Text3D/Public/Text3DTypes.h
+++ b/Engine/Plugins/VirtualProduction/Text3D/Source/Text3D/Public/Text3DTypes.h
@@ -3,6 +3,7 @@
 #pragma once
 
 #include "Containers/Array.h"
+#include "Math/Box.h"
 #include "Misc/EnumClassFlags.h"
 #include "UObject/ObjectPtr.h"
 #include "Text3DTypes.generated.h"
