Text3D optimization by adding option to disable collision. With collision disabled, the performance improvement is substantial.
Also changed components to now be added to the scene in batch rather than individually for further improvement in performance.

#Synced-CL 50556942

diff --git a/Engine/Plugins/VirtualProduction/Text3D/Source/Text3D/Private/Extensions/Text3DDefaultGeometryExtension.cpp b/Engine/Plugins/VirtualProduction/Text3D/Source/Text3D/Private/Extensions/Text3DDefaultGeometryExtension.cpp
--- a/Engine/Plugins/VirtualProduction/Text3D/Source/Text3D/Private/Extensions/Text3DDefaultGeometryExtension.cpp
+++ b/Engine/Plugins/VirtualProduction/Text3D/Source/Text3D/Private/Extensions/Text3DDefaultGeometryExtension.cpp
@@ -71,6 +71,17 @@ void UText3DDefaultGeometryExtension::SetPivotVAlignment(EText3DVerticalTextAlig
 	OnGeometryOptionsChanged();
 }
 
+void UText3DDefaultGeometryExtension::SetCollisionEnabled(ECollisionEnabled::Type InCollisionEnabled)
+{
+	if (CollisionEnabled == InCollisionEnabled)
+	{
+		return;
+	}
+
+	CollisionEnabled = InCollisionEnabled;
+	OnCollisionEnabledChanged();
+}
+
 void UText3DDefaultGeometryExtension::SetExtrude(const float Value)
 {
 	const float NewValue = FMath::Max(0.0f, Value);
@@ -152,6 +163,10 @@ void UText3DDefaultGeometryExtension::PostEditChangeProperty(FPropertyChangedEve
 	{
 		OnGeometryOptionsChanged();
 	}
+	else if (InEvent.GetMemberPropertyName() == GET_MEMBER_NAME_CHECKED(UText3DDefaultGeometryExtension, CollisionEnabled))
+	{
+		OnCollisionEnabledChanged();
+	}
 }
 #endif
 
@@ -165,6 +180,11 @@ EText3DVerticalTextAlignment UText3DDefaultGeometryExtension::GetGlyphVAlignment
 	return PivotVAlignment;
 }
 
+ECollisionEnabled::Type UText3DDefaultGeometryExtension::GetGlyphCollisionEnabled() const
+{
+	return CollisionEnabled;
+}
+
 const FGlyphMeshParameters* UText3DDefaultGeometryExtension::GetGlyphMeshParameters() const
 {
 	return &GlyphMeshParameters;
@@ -250,6 +270,11 @@ void UText3DDefaultGeometryExtension::OnGeometryOptionsChanged()
 	RequestUpdate(EText3DRendererFlags::All);
 }
 
+void UText3DDefaultGeometryExtension::OnCollisionEnabledChanged()
+{
+	RequestUpdate(EText3DRendererFlags::Geometry);
+}
+
 float UText3DDefaultGeometryExtension::GetMaxBevel() const
 {
 	return Extrude / 2.0f;
diff --git a/Engine/Plugins/VirtualProduction/Text3D/Source/Text3D/Private/Renderers/DynamicMesh/Text3DDynamicMeshRenderer.cpp b/Engine/Plugins/VirtualProduction/Text3D/Source/Text3D/Private/Renderers/DynamicMesh/Text3DDynamicMeshRenderer.cpp
--- a/Engine/Plugins/VirtualProduction/Text3D/Source/Text3D/Private/Renderers/DynamicMesh/Text3DDynamicMeshRenderer.cpp
+++ b/Engine/Plugins/VirtualProduction/Text3D/Source/Text3D/Private/Renderers/DynamicMesh/Text3DDynamicMeshRenderer.cpp
@@ -96,6 +96,10 @@ void UText3DDynamicMeshRenderer::OnUpdate(const UE::Text3D::Renderer::FUpdatePar
 
 	if (InParameters.CurrentFlag == EText3DRendererFlags::Geometry)
 	{
+		UText3DGeometryExtensionBase* const GeometryExtension = TextComponent->GetGeometryExtension();
+		DynamicMeshComponent->bAlwaysCreatePhysicsState = false;
+		DynamicMeshComponent->SetCollisionEnabled(GeometryExtension ? GeometryExtension->GetGlyphCollisionEnabled() : ECollisionEnabled::QueryAndPhysics);
+
 		AllocateGlyphMeshData(MeshEditor, TextComponent->GetCharacterCount());
 		
 		TextComponent->ForEachCharacter([this, &MeshEditor](UText3DCharacterBase* InCharacter, uint16 InIndex, uint16 InCount)
diff --git a/Engine/Plugins/VirtualProduction/Text3D/Source/Text3D/Private/Renderers/StaticMeshes/Text3DStaticMeshesRenderer.cpp b/Engine/Plugins/VirtualProduction/Text3D/Source/Text3D/Private/Renderers/StaticMeshes/Text3DStaticMeshesRenderer.cpp
--- a/Engine/Plugins/VirtualProduction/Text3D/Source/Text3D/Private/Renderers/StaticMeshes/Text3DStaticMeshesRenderer.cpp
+++ b/Engine/Plugins/VirtualProduction/Text3D/Source/Text3D/Private/Renderers/StaticMeshes/Text3DStaticMeshesRenderer.cpp
@@ -1,16 +1,17 @@
 // Copyright Epic Games, Inc. All Rights Reserved.
 
 #include "Renderers/StaticMeshes/Text3DStaticMeshesRenderer.h"
-
 #include "Characters/Text3DCharacterBase.h"
 #include "Components/SceneComponent.h"
 #include "Components/StaticMeshComponent.h"
 #include "Engine/StaticMesh.h"
+#include "Engine/World.h"
 #include "Extensions/Text3DGeometryExtensionBase.h"
 #include "Extensions/Text3DLayoutExtensionBase.h"
 #include "Extensions/Text3DMaterialExtensionBase.h"
 #include "Extensions/Text3DRenderingExtensionBase.h"
 #include "Logs/Text3DLogs.h"
+#include "SceneInterface.h"
 #include "Text3DComponent.h"
 #include "Text3DInternalTypes.h"
 #include "Utilities/Text3DUtilities.h"
@@ -58,17 +59,45 @@ void UText3DStaticMeshesRenderer::OnUpdate(const UE::Text3D::Renderer::FUpdatePa
 	{
 		AllocateComponents(TextComponent->GetCharacterCount());
 
-		TextComponent->ForEachCharacter([this](UText3DCharacterBase* InCharacter, uint16 InIndex, uint16 InCount)
+		UWorld* const World = TextComponent->GetWorld();
+
+		const TArrayView<UPrimitiveComponent*> StaticMeshComponents = MakeArrayView(reinterpret_cast<UPrimitiveComponent**>(CharacterMeshes.GetData()), CharacterMeshes.Num());
+		auto SetBulkReregister = [&StaticMeshComponents](bool bInNewValue)
+			{
+				for (UPrimitiveComponent* Component : StaticMeshComponents)
+				{
+					if (Component)
+					{
+						Component->bBulkReregister = bInNewValue;
+					}
+				}
+			};
+
+		if (World && !StaticMeshComponents.IsEmpty())
+		{
+			World->Scene->BatchRemovePrimitives(StaticMeshComponents);
+			World->Scene->BatchReleasePrimitives(StaticMeshComponents);
+			SetBulkReregister(true);
+		}
+
+		UText3DGeometryExtensionBase* const GeometryExtension = TextComponent->GetGeometryExtension();
+		const ECollisionEnabled::Type CollisionEnabled = GeometryExtension ? GeometryExtension->GetGlyphCollisionEnabled() : ECollisionEnabled::QueryAndPhysics;
+
+		TextComponent->ForEachCharacter([this, CollisionEnabled](UText3DCharacterBase* InCharacter, uint16 InIndex, uint16 InCount)
 		{
 			if (CharacterMeshes.IsValidIndex(InIndex))
 			{
+				UStaticMeshComponent* const StaticMeshComponent = CharacterMeshes[InIndex];
+				StaticMeshComponent->bAlwaysCreatePhysicsState = false;
+				StaticMeshComponent->SetCollisionEnabled(CollisionEnabled);
+
 				if (const FText3DCachedMesh* CachedMesh = InCharacter->GetGlyphMesh())
 				{
-					CharacterMeshes[InIndex]->SetStaticMesh(CachedMesh->StaticMesh);
+					StaticMeshComponent->SetStaticMesh(CachedMesh->StaticMesh);
 				}
 				else
 				{
-					CharacterMeshes[InIndex]->SetStaticMesh(nullptr);
+					StaticMeshComponent->SetStaticMesh(nullptr);
 				}
 			}
 			else
@@ -77,10 +106,18 @@ void UText3DStaticMeshesRenderer::OnUpdate(const UE::Text3D::Renderer::FUpdatePa
 			}
 		});
 
+		if (World && !StaticMeshComponents.IsEmpty())
+		{
+			World->Scene->BatchAddPrimitives(StaticMeshComponents);
+			SetBulkReregister(false);
+		}
+
 		RefreshBounds();
 	}
 	else if (InParameters.CurrentFlag == EText3DRendererFlags::Layout)
 	{
+		TRACE_CPUPROFILER_EVENT_SCOPE(UText3DStaticMeshesRenderer::OnUpdate_Layout);
+
 		const UText3DLayoutExtensionBase* LayoutExtension = TextComponent->GetLayoutExtension();
 
 		TextRoot->SetRelativeScale3D(LayoutExtension->GetTextScale());
@@ -104,6 +141,8 @@ void UText3DStaticMeshesRenderer::OnUpdate(const UE::Text3D::Renderer::FUpdatePa
 	}
 	else if (InParameters.CurrentFlag == EText3DRendererFlags::Material)
 	{
+		TRACE_CPUPROFILER_EVENT_SCOPE(UText3DStaticMeshesRenderer::OnUpdate_Material);
+
 		using namespace UE::Text3D::Material;
 
 		const UText3DMaterialExtensionBase* MaterialExtension = TextComponent->GetMaterialExtension();
@@ -145,6 +184,8 @@ void UText3DStaticMeshesRenderer::OnUpdate(const UE::Text3D::Renderer::FUpdatePa
 	}
 	else if (InParameters.CurrentFlag == EText3DRendererFlags::Visibility)
 	{
+		TRACE_CPUPROFILER_EVENT_SCOPE(UText3DStaticMeshesRenderer::OnUpdate_Visibility);
+
 		const UText3DRenderingExtensionBase* RenderingExtension = TextComponent->GetRenderingExtension();
 
 		TextComponent->ForEachCharacter([this, TextComponent, RenderingExtension](UText3DCharacterBase* InCharacter, uint16 InIndex, uint16 InCount)
@@ -219,6 +260,8 @@ void UText3DStaticMeshesRenderer::AllocateComponents(int32 InCount)
 		return;
 	}
 
+	TRACE_CPUPROFILER_EVENT_SCOPE(UText3DStaticMeshesRenderer::AllocateComponents);
+
 	CharacterMeshes.Reserve(InCount);
 
 	const int32 RemainingCharacterCount = CharacterMeshes.Num() - InCount;
@@ -252,7 +295,8 @@ void UText3DStaticMeshesRenderer::AllocateComponents(int32 InCount)
 	}
 	else if (RemainingCharacterCount < 0)
 	{
-		for (int32 CharacterIndex = 0; CharacterIndex < FMath::Abs(RemainingCharacterCount); ++CharacterIndex)
+		const int32 RemainingCharacterCountPositive = FMath::Abs(RemainingCharacterCount);
+		for (int32 CharacterIndex = 0; CharacterIndex < RemainingCharacterCountPositive; ++CharacterIndex)
 		{
 			UStaticMeshComponent* MeshComponent = nullptr;
 			if (!CharacterMeshesPool.IsEmpty())
@@ -262,8 +306,7 @@ void UText3DStaticMeshesRenderer::AllocateComponents(int32 InCount)
 
 			if (!MeshComponent)
 			{
-				const FName StaticMeshComponentName = MakeUniqueObjectName(this, UStaticMeshComponent::StaticClass(), FName(TEXT("CharacterMesh")));
-				MeshComponent = NewObject<UStaticMeshComponent>(this, StaticMeshComponentName);
+				MeshComponent = NewObject<UStaticMeshComponent>(this);
 			}
 
 			if (MeshComponent)
@@ -307,11 +350,6 @@ const TArray<UStaticMeshComponent*>& UText3DStaticMeshesRenderer::GetGlyphMeshCo
 #if WITH_EDITOR
 void UText3DStaticMeshesRenderer::ConvertToStaticMesh()
 {
-	if (!IsValid(this))
-	{
-		return;
-	}
-
 	const UText3DComponent* Text3DComponent = GetText3DComponent();
 	if (!IsValid(Text3DComponent) || !IsValid(Text3DComponent->GetOwner()))
 	{
diff --git a/Engine/Plugins/VirtualProduction/Text3D/Source/Text3D/Public/Extensions/Text3DDefaultGeometryExtension.h b/Engine/Plugins/VirtualProduction/Text3D/Source/Text3D/Public/Extensions/Text3DDefaultGeometryExtension.h
--- a/Engine/Plugins/VirtualProduction/Text3D/Source/Text3D/Public/Extensions/Text3DDefaultGeometryExtension.h
+++ b/Engine/Plugins/VirtualProduction/Text3D/Source/Text3D/Public/Extensions/Text3DDefaultGeometryExtension.h
@@ -103,6 +103,13 @@ public:
 		return PivotVAlignment;
 	}
 	
+	UFUNCTION()
+	TEXT3D_API void SetCollisionEnabled(ECollisionEnabled::Type InCollisionEnabled);
+	ECollisionEnabled::Type GetCollisionEnabled() const
+	{
+		return CollisionEnabled;
+	}
+
 protected:
 	//~ Begin UObject
 #if WITH_EDITOR
@@ -114,6 +121,7 @@ protected:
 	//~ Begin UText3DGeometryExtensionBase
 	virtual EText3DHorizontalTextAlignment GetGlyphHAlignment() const override;
 	virtual EText3DVerticalTextAlignment GetGlyphVAlignment() const override;
+	virtual ECollisionEnabled::Type GetGlyphCollisionEnabled() const override;
 	virtual const FGlyphMeshParameters* GetGlyphMeshParameters() const override;
 	//~ End UText3DGeometryExtensionBase
 
@@ -123,6 +131,7 @@ protected:
 	//~ End UText3DExtensionBase
 
 	void OnGeometryOptionsChanged();
+	void OnCollisionEnabledChanged();
 
 	float GetMaxBevel() const;
 	FVector GetPivotOffset() const;
@@ -163,6 +172,9 @@ protected:
 	UPROPERTY(Getter, Setter, meta=(ValidEnumValues="Bottom"))
 	EText3DVerticalTextAlignment PivotVAlignment = EText3DVerticalTextAlignment::Bottom;
 
+	UPROPERTY(EditAnywhere, Setter, Getter, Category="Geometry", meta=(AllowPrivateAccess = "true"))
+	TEnumAsByte<ECollisionEnabled::Type> CollisionEnabled = ECollisionEnabled::QueryAndPhysics;
+
 private:
 	FGlyphMeshParameters GlyphMeshParameters;
 };
\ No newline at end of file
diff --git a/Engine/Plugins/VirtualProduction/Text3D/Source/Text3D/Public/Extensions/Text3DGeometryExtensionBase.h b/Engine/Plugins/VirtualProduction/Text3D/Source/Text3D/Public/Extensions/Text3DGeometryExtensionBase.h
--- a/Engine/Plugins/VirtualProduction/Text3D/Source/Text3D/Public/Extensions/Text3DGeometryExtensionBase.h
+++ b/Engine/Plugins/VirtualProduction/Text3D/Source/Text3D/Public/Extensions/Text3DGeometryExtensionBase.h
@@ -2,6 +2,7 @@
 
 #pragma once
 
+#include "Engine/EngineTypes.h"
 #include "Text3DExtensionBase.h"
 #include "Text3DTypes.h"
 #include "Text3DGeometryExtensionBase.generated.h"
@@ -40,6 +41,11 @@ public:
 		return EText3DVerticalTextAlignment::Bottom;
 	}
 
+	virtual ECollisionEnabled::Type GetGlyphCollisionEnabled() const
+	{
+		return ECollisionEnabled::QueryAndPhysics;
+	}
+
 	virtual const FGlyphMeshParameters* GetGlyphMeshParameters() const
 	{
 		return nullptr;
