Text3D: Fix issue with Text3D scene proxies being null

#Synced-CL 50556942

diff --git a/Engine/Plugins/VirtualProduction/Text3D/Source/Text3D/Private/Renderers/StaticMeshes/Text3DStaticMeshesRenderer.cpp b/Engine/Plugins/VirtualProduction/Text3D/Source/Text3D/Private/Renderers/StaticMeshes/Text3DStaticMeshesRenderer.cpp
--- a/Engine/Plugins/VirtualProduction/Text3D/Source/Text3D/Private/Renderers/StaticMeshes/Text3DStaticMeshesRenderer.cpp
+++ b/Engine/Plugins/VirtualProduction/Text3D/Source/Text3D/Private/Renderers/StaticMeshes/Text3DStaticMeshesRenderer.cpp
@@ -22,6 +22,17 @@
 
 #define LOCTEXT_NAMESPACE "Text3DStaticMeshesRenderer"
 
+namespace UE::Text3D::Private
+{
+	// Static mesh components will be added in batch to the FScene, rather than one by one
+	static bool bBatchRegisterStaticMeshes = true;
+	static FAutoConsoleVariableRef CVarBatchRegisterStaticMeshes(
+		TEXT("Text3D.BatchRegisterStaticMeshes"),
+		bBatchRegisterStaticMeshes,
+		TEXT("Determines whether static mesh component primitives should be batch added to the scene, rather than one by one. True to enable it, false to revert back to legacy behavior"),
+		ECVF_Default);
+}
+
 void UText3DStaticMeshesRenderer::OnCreate()
 {
 	UText3DComponent* TextComponent = GetText3DComponent();
@@ -62,6 +73,13 @@ void UText3DStaticMeshesRenderer::OnUpdate(const UE::Text3D::Renderer::FUpdatePa
 		UWorld* const World = TextComponent->GetWorld();
 
 		const TArrayView<UPrimitiveComponent*> StaticMeshComponents = MakeArrayView(reinterpret_cast<UPrimitiveComponent**>(CharacterMeshes.GetData()), CharacterMeshes.Num());
+
+		const bool bBatchRegisterStaticMeshes = UE::Text3D::Private::bBatchRegisterStaticMeshes
+			&& !StaticMeshComponents.ContainsByPredicate([](UPrimitiveComponent* InComponent)
+				{
+					return !InComponent || !InComponent->IsRegistered();
+				});
+
 		auto SetBulkReregister = [&StaticMeshComponents](bool bInNewValue)
 			{
 				for (UPrimitiveComponent* Component : StaticMeshComponents)
@@ -73,7 +91,7 @@ void UText3DStaticMeshesRenderer::OnUpdate(const UE::Text3D::Renderer::FUpdatePa
 				}
 			};
 
-		if (World && !StaticMeshComponents.IsEmpty())
+		if (bBatchRegisterStaticMeshes && World && !StaticMeshComponents.IsEmpty())
 		{
 			World->Scene->BatchRemovePrimitives(StaticMeshComponents);
 			World->Scene->BatchReleasePrimitives(StaticMeshComponents);
@@ -106,7 +124,7 @@ void UText3DStaticMeshesRenderer::OnUpdate(const UE::Text3D::Renderer::FUpdatePa
 			}
 		});
 
-		if (World && !StaticMeshComponents.IsEmpty())
+		if (bBatchRegisterStaticMeshes && World && !StaticMeshComponents.IsEmpty())
 		{
 			World->Scene->BatchAddPrimitives(StaticMeshComponents);
 			SetBulkReregister(false);
