Text3D glyph optimization:
Removed use of thread-safe TSharedPtr to reduce heap allocations and atomic ref counting when building glyph as this was not needed.
In addition, the glyph mesh building is now done across multiple threads for performance

#Synced-CL 50556942

diff --git a/Engine/Plugins/VirtualProduction/Text3D/Source/Text3D/Private/Extensions/Text3DDefaultGeometryExtension.cpp b/Engine/Plugins/VirtualProduction/Text3D/Source/Text3D/Private/Extensions/Text3DDefaultGeometryExtension.cpp
--- a/Engine/Plugins/VirtualProduction/Text3D/Source/Text3D/Private/Extensions/Text3DDefaultGeometryExtension.cpp
+++ b/Engine/Plugins/VirtualProduction/Text3D/Source/Text3D/Private/Extensions/Text3DDefaultGeometryExtension.cpp
@@ -7,6 +7,7 @@
 #include "Engine/StaticMesh.h"
 #include "Logs/Text3DLogs.h"
 #include "MeshDescription.h"
+#include "Renderers/Text3DRendererBase.h"
 #include "Subsystems/Text3DEngineSubsystem.h"
 #include "Text3DComponent.h"
 
@@ -211,6 +212,7 @@ EText3DExtensionResult UText3DDefaultGeometryExtension::PreRendererUpdate(const
 	const float GlyphBevel = FMath::TruncToFloat(GetBevel() * 10.f) / 10.f;
 	const float GlyphOutline = FMath::TruncToFloat(GetOutline() * 10.f) / 10.f;
 
+	const UText3DRendererBase* TextRenderer = TextComponent->GetTextRenderer();
 	GlyphMeshParameters =
 	{
 		.Extrude = GlyphExtrude,
@@ -221,6 +223,7 @@ EText3DExtensionResult UText3DDefaultGeometryExtension::PreRendererUpdate(const
 		.OutlineExpand = GlyphOutline,
 		.OutlineType = GetOutlineType(),
 		.PivotOffset = GetPivotOffset(),
+		.MeshType = TextRenderer ? TextRenderer->GetMeshType() : EText3DMeshType::Unknown,
 	};
 
 	// Keep to clean after
diff --git a/Engine/Plugins/VirtualProduction/Text3D/Source/Text3D/Private/Extensions/Text3DDefaultLayoutExtension.cpp b/Engine/Plugins/VirtualProduction/Text3D/Source/Text3D/Private/Extensions/Text3DDefaultLayoutExtension.cpp
--- a/Engine/Plugins/VirtualProduction/Text3D/Source/Text3D/Private/Extensions/Text3DDefaultLayoutExtension.cpp
+++ b/Engine/Plugins/VirtualProduction/Text3D/Source/Text3D/Private/Extensions/Text3DDefaultLayoutExtension.cpp
@@ -4,6 +4,7 @@
 
 #include "Algo/Accumulate.h"
 #include "Algo/Count.h"
+#include "Async/ParallelFor.h"
 #include "Characters/Text3DCharacterBase.h"
 #include "Engine/Font.h"
 #include "Extensions/Text3DCharacterExtensionBase.h"
@@ -15,9 +16,11 @@
 #include "LayoutBuilders/Text3DLayoutShaper.h"
 #include "LayoutBuilders/Text3DShapedGlyphText.h"
 #include "Logs/Text3DLogs.h"
+#include "Materials/Material.h"
 #include "Misc/EnumerateRange.h"
-#include "Styling/StyleDefaults.h"
+#include "Settings/Text3DProjectSettings.h"
 #include "Styling/SlateTypes.h"
+#include "Styling/StyleDefaults.h"
 #include "Subsystems/Text3DEngineSubsystem.h"
 #include "Text3DComponent.h"
 
@@ -260,8 +263,7 @@ EText3DExtensionResult UText3DDefaultLayoutExtension::PreRendererUpdate(const UE
 	GlyphText->Kernings.SetNum(CharacterCount);
 	GlyphText->FontFaces.SetNum(CharacterCount);
 
-	const UText3DGeometryExtensionBase* GeometryExtension = Text3DComponent->GetGeometryExtension();
-	uint16 CharacterIndex = 0;
+	int32 CharacterIndex = 0;
 
 #if WITH_FREETYPE
 	for (const TConstEnumerateRef<FGlyphLine> GlyphLine : EnumerateRange(GlyphText->Lines))
@@ -285,10 +287,6 @@ EText3DExtensionResult UText3DDefaultLayoutExtension::PreRendererUpdate(const UE
 			}
 		}
 	}
-#else
-	UE_LOG(LogText3D, Warning, TEXT("FreeType is not available, cannot proceed without it"))
-	return EText3DExtensionResult::Failed;
-#endif
 
 	GlyphText->Tracking = Tracking;
 	GlyphText->WordSpacing = WordSpacing;
@@ -297,8 +295,90 @@ EText3DExtensionResult UText3DDefaultLayoutExtension::PreRendererUpdate(const UE
 	GlyphText->CalculateWidth();
 	CalculateTextScale();
 
+	UText3DEngineSubsystem* const TextSubsystem = UText3DEngineSubsystem::Get();
+	if (!ensure(TextSubsystem))
+	{
+		return EText3DExtensionResult::Failed;
+	}
+
+	const UText3DGeometryExtensionBase* const GeometryExtension = Text3DComponent->GetGeometryExtension();
+	if (!ensure(GeometryExtension))
+	{
+		return EText3DExtensionResult::Failed;
+	}
+
+	const FGlyphMeshParameters& GlyphMeshParameters = *GeometryExtension->GetGlyphMeshParameters();
+
 	CharacterIndex = 0;
-	UText3DEngineSubsystem* TextSubsystem = UText3DEngineSubsystem::Get();
+
+	if (EnumHasAnyFlags(InParameters.UpdateFlags, EText3DRendererFlags::Geometry))
+	{
+		struct FBuildGlyphMeshDesc
+		{
+			FText3DGlyphContourNode RootNode;
+			TSharedPtr<FFreeTypeFace> FontFace;
+			uint32 GlyphIndex = 0;
+		};
+		TArray<FBuildGlyphMeshDesc> BuildGlyphMeshDescs;
+
+		for (const TConstEnumerateRef<FGlyphLine> GlyphLine : EnumerateRange(GlyphText->Lines))
+		{
+			for (const TConstEnumerateRef<FGlyphEntry> GlyphEntry : EnumerateRange(GlyphLine->Glyphs))
+			{
+				if (!FGlyphText::IsValidGlyph(*GlyphEntry))
+				{
+					continue;
+				}
+
+				if (!Text3DComponent->GetCharacter(CharacterIndex))
+				{
+					UE_LOG(LogText3D, Warning, TEXT("Invalid character object returned for index %i after layout calculations, cannot proceed"), CharacterIndex)
+					return EText3DExtensionResult::Failed;
+				}
+
+				const TSharedPtr<FFreeTypeFace>& FontFace = GlyphText->FontFaces[CharacterIndex++];
+				if (FText3DFontFaceCache* FontFaceCache = TextSubsystem->FindOrAddCachedFontFace(FontFace))
+				{
+					if (TOptional<FText3DGlyphContourNode> RootNode = FontFaceCache->PrepareCachedMesh(GlyphEntry->Entry.GlyphIndex, GlyphMeshParameters))
+					{
+						BuildGlyphMeshDescs.Emplace(MoveTemp(*RootNode), FontFace, GlyphEntry->Entry.GlyphIndex);
+					}
+				}
+				else
+				{
+					UE_LOG(LogText3D, Warning, TEXT("Failed to get cached font data for '%u %i' in Text3D geometry extension"), FText3DFontFaceCache::GetFontFaceHash(FontFace), GlyphEntry->Entry.GlyphIndex);
+					return EText3DExtensionResult::Failed;
+				}
+			}
+		}
+
+		TRACE_CPUPROFILER_EVENT_SCOPE(UText3DDefaultLayoutExtension::BuildGlyphMeshes);
+		const FText3DFontFaceCache::FBuildGlyphMeshSharedDesc SharedParams
+			{
+				.TextSubsystem = TextSubsystem,
+				.GlyphMeshParameters = &GlyphMeshParameters,
+				.DefaultMaterial = GetDefault<UText3DProjectSettings>()->GetDefaultMaterial(),
+			};
+
+		ParallelFor(BuildGlyphMeshDescs.Num(), 
+			[&SharedParams, &BuildGlyphMeshDescs](int32 InTaskInfoIndex)
+			{
+				const FBuildGlyphMeshDesc& BuildGlyphMeshDesc = BuildGlyphMeshDescs[InTaskInfoIndex];
+
+				if (FText3DFontFaceCache* FontFaceCache = SharedParams.TextSubsystem->FindCachedFontFace(BuildGlyphMeshDesc.FontFace))
+				{
+					const FText3DFontFaceCache::FBuildGlyphMeshInstanceDesc InstanceParams
+						{
+							.RootNode = &BuildGlyphMeshDesc.RootNode,
+							.GlyphIndex = BuildGlyphMeshDesc.GlyphIndex,
+						};
+					FontFaceCache->BuildGlyphMesh(SharedParams, InstanceParams);
+				}
+			});
+	}
+
+	CharacterIndex = 0;
+
 	for (TConstEnumerateRef<FGlyphLine> GlyphLine : EnumerateRange(GlyphText->Lines))
 	{
 		FVector Location = GetLineLocation(GlyphLine.GetIndex());
@@ -316,26 +396,19 @@ EText3DExtensionResult UText3DDefaultLayoutExtension::PreRendererUpdate(const UE
 
 				if (EnumHasAnyFlags(InParameters.UpdateFlags, EText3DRendererFlags::Geometry))
 				{
-#if WITH_FREETYPE
-					TSharedPtr<FFreeTypeFace> CharTypeface = GlyphText->FontFaces[CharacterIndex];
-					const uint32 FontFaceHash = FText3DFontFaceCache::GetFontFaceHash(CharTypeface);
+					TSharedPtr<FFreeTypeFace> FontFace = GlyphText->FontFaces[CharacterIndex];
 					// Cache font face to keep it alive and use inner FT_Face to build glyphs
-					FText3DFontFaceCache* FontFaceCache = TextSubsystem->FindOrAddCachedFontFace(CharTypeface);
-#else
-					uint32 FontFaceHash = 0;
-					FText3DFontFaceCache* FontFaceCache = nullptr;
-#endif
+					FText3DFontFaceCache* FontFaceCache = TextSubsystem->FindOrAddCachedFontFace(FontFace);
 
 					if (FontFaceCache)
 					{
-						// Load glyph here to get proper mesh metrics
-						UE::Text3D::Geometry::FCachedFontFaceGlyphHandle GlyphHandle(FontFaceCache, GlyphEntry->Entry.GlyphIndex, *GeometryExtension->GetGlyphMeshParameters());
+						UE::Text3D::Geometry::FCachedFontFaceGlyphHandle GlyphHandle(*FontFaceCache, GlyphEntry->Entry.GlyphIndex, GlyphMeshParameters);
 						Character->SetFontFaceGlyphHandle(GlyphHandle);
 						Character->SetStyleTag(GlyphEntry->StyleTag);
 					}
 					else
 					{
-						UE_LOG(LogText3D, Warning, TEXT("Failed to get cached font data for '%u %i' in Text3D geometry extension"), FontFaceHash, GlyphEntry->Entry.GlyphIndex);
+						UE_LOG(LogText3D, Warning, TEXT("Failed to get cached font data for '%u %i' in Text3D geometry extension"), FText3DFontFaceCache::GetFontFaceHash(FontFace), GlyphEntry->Entry.GlyphIndex);
 						return EText3DExtensionResult::Failed;
 					}
 				}
@@ -401,6 +474,10 @@ EText3DExtensionResult UText3DDefaultLayoutExtension::PreRendererUpdate(const UE
 	}
 
 	return EText3DExtensionResult::Finished;
+#else
+	UE_LOG(LogText3D, Warning, TEXT("FreeType is not available, cannot proceed without it"))
+	return EText3DExtensionResult::Failed;
+#endif // WITH_FREETYPE
 }
 
 EText3DExtensionResult UText3DDefaultLayoutExtension::PostRendererUpdate(const UE::Text3D::Renderer::FUpdateParameters& InParameters)
diff --git a/Engine/Plugins/VirtualProduction/Text3D/Source/Text3D/Private/GeometryBuilders/Text3DGlyph.cpp b/Engine/Plugins/VirtualProduction/Text3D/Source/Text3D/Private/GeometryBuilders/Text3DGlyph.cpp
--- a/Engine/Plugins/VirtualProduction/Text3D/Source/Text3D/Private/GeometryBuilders/Text3DGlyph.cpp
+++ b/Engine/Plugins/VirtualProduction/Text3D/Source/Text3D/Private/GeometryBuilders/Text3DGlyph.cpp
@@ -26,10 +26,11 @@ FText3DGlyph::FText3DGlyph()
 	}
 }
 
-void FText3DGlyph::Build(FText3DCachedMesh& InMesh, UMaterialInterface* InDefaultMaterial)
+void FText3DGlyph::Build(FText3DCachedMesh& InMesh, EText3DMeshType InMeshType, UMaterialInterface* InDefaultMaterial)
 {
+	// Both dynamic and static build static meshes
 	checkf(InMesh.StaticMesh, TEXT("Invalid static mesh object provided to the build function"));
-	checkf(InMesh.DynamicMesh, TEXT("Invalid dynamic mesh object provided to the build function"));
+	checkf(InMeshType != EText3DMeshType::Dynamic || InMesh.DynamicMesh, TEXT("Invalid dynamic mesh object provided to the build function"));
 
 	auto AddMaterial = [](UStaticMesh* InMesh, UMaterialInterface* InMaterial, int32 InIndex)
 	{
@@ -63,12 +64,14 @@ void FText3DGlyph::Build(FText3DCachedMesh& InMesh, UMaterialInterface* InDefaul
 		TArray<const FMeshDescription*> MeshDescriptions;
 		MeshDescriptions.Add(&MeshDescription);
 
+		constexpr bool bAllowCpuAccess = true;
+
 		UStaticMesh::FBuildMeshDescriptionsParams BuildParams;
 		BuildParams.bCommitMeshDescription = true;
 		BuildParams.bFastBuild = true;
-		BuildParams.bAllowCpuAccess = true;
+		BuildParams.bAllowCpuAccess = bAllowCpuAccess;
 
-		InMesh.StaticMesh->bAllowCPUAccess = true;
+		InMesh.StaticMesh->bAllowCPUAccess = bAllowCpuAccess;
 		InMesh.StaticMesh->BuildFromMeshDescriptions(MeshDescriptions, BuildParams);
 
 		// Setup collision
@@ -78,6 +81,8 @@ void FText3DGlyph::Build(FText3DCachedMesh& InMesh, UMaterialInterface* InDefaul
 		}
 	}
 
+	// Dynamic mesh building
+	if (InMeshType == EText3DMeshType::Dynamic)
 	{
 		static constexpr FGeometryScriptMeshReadLOD StaticMeshLOD
 		{
@@ -96,6 +101,11 @@ void FText3DGlyph::Build(FText3DCachedMesh& InMesh, UMaterialInterface* InDefaul
 	}
 }
 
+const FMeshDescription& FText3DGlyph::GetMeshDescription() const
+{
+	return MeshDescription;
+}
+
 FMeshDescription& FText3DGlyph::GetMeshDescription()
 {
 	return MeshDescription;
diff --git a/Engine/Plugins/VirtualProduction/Text3D/Source/Text3D/Private/GeometryBuilders/Text3DGlyph.h b/Engine/Plugins/VirtualProduction/Text3D/Source/Text3D/Private/GeometryBuilders/Text3DGlyph.h
--- a/Engine/Plugins/VirtualProduction/Text3D/Source/Text3D/Private/GeometryBuilders/Text3DGlyph.h
+++ b/Engine/Plugins/VirtualProduction/Text3D/Source/Text3D/Private/GeometryBuilders/Text3DGlyph.h
@@ -25,8 +25,9 @@ class FText3DGlyph
 public:
 	FText3DGlyph();
 
-	void Build(FText3DCachedMesh& InMesh, UMaterialInterface* InDefaultMaterial);
+	void Build(FText3DCachedMesh& InMesh, EText3DMeshType InMeshType, UMaterialInterface* InDefaultMaterial);
 
+	const FMeshDescription& GetMeshDescription() const;
 	FMeshDescription& GetMeshDescription();
 	FStaticMeshAttributes& GetStaticMeshAttributes();
 	TText3DGroupList& GetGroups();
diff --git a/Engine/Plugins/VirtualProduction/Text3D/Source/Text3D/Private/GeometryBuilders/Text3DGlyphArrangement.cpp b/Engine/Plugins/VirtualProduction/Text3D/Source/Text3D/Private/GeometryBuilders/Text3DGlyphArrangement.cpp
--- a/Engine/Plugins/VirtualProduction/Text3D/Source/Text3D/Private/GeometryBuilders/Text3DGlyphArrangement.cpp
+++ b/Engine/Plugins/VirtualProduction/Text3D/Source/Text3D/Private/GeometryBuilders/Text3DGlyphArrangement.cpp
@@ -17,9 +17,9 @@ void FText3DGlyphArrangement::Insert(const FSegment2f& Segment)
     InsertSegment({static_cast<double>(A.X), static_cast<double>(A.Y)}, {static_cast<double>(B.X), static_cast<double>(B.Y)}, VertexSnapTol);
 }
 
-int FText3DGlyphArrangement::InsertPoint(const FVector2d &InPoint, double InTol)
+int32 FText3DGlyphArrangement::InsertPoint(const FVector2d &InPoint, double InTol)
 {
-    int PIdx = FindExistingVertex(InPoint);
+    int32 PIdx = FindExistingVertex(InPoint);
     if (PIdx > -1)
     {
         return -1;
@@ -28,8 +28,8 @@ int FText3DGlyphArrangement::InsertPoint(const FVector2d &InPoint, double InTol)
     // TODO: currently this tries to add the vertex on the closest edge below tolerance; we should instead insert at *every* edge below tolerance!  ... but that is more inconvenient to write
     FVector2d x = FVector2d::Zero(), y = FVector2d::Zero();
     double ClosestDistSq = InTol*InTol;
-    int FoundEdgeToSplit = -1;
-    for (int EID = 0, ExistingEdgeMax = Graph.MaxEdgeID(); EID < ExistingEdgeMax; EID++)
+    int32 FoundEdgeToSplit = -1;
+    for (int32 EID = 0, ExistingEdgeMax = Graph.MaxEdgeID(); EID < ExistingEdgeMax; EID++)
     {
         if (!Graph.IsEdge(EID))
         {
@@ -51,20 +51,20 @@ int FText3DGlyphArrangement::InsertPoint(const FVector2d &InPoint, double InTol)
         EMeshResult result = Graph.SplitEdge(FoundEdgeToSplit, splitInfo);
         ensureMsgf(result == EMeshResult::Ok, TEXT("insert_into_segment: edge split failed?"));
         Graph.SetVertex(splitInfo.VNew, InPoint);
-        PointHash.InsertPointUnsafe(splitInfo.VNew, InPoint);
+        PointHash.InsertPoint(splitInfo.VNew, InPoint);
         return splitInfo.VNew;
     }
 
-    int VID = Graph.AppendVertex(InPoint);
-    PointHash.InsertPointUnsafe(VID, InPoint);
+    int32 VID = Graph.AppendVertex(InPoint);
+    PointHash.InsertPoint(VID, InPoint);
     return VID;
 }
 
 bool FText3DGlyphArrangement::InsertSegment(FVector2d InA, FVector2d InB, double InTol)
 {
     // handle degenerate edges
-    int a_idx = FindExistingVertex(InA);
-    int b_idx = FindExistingVertex(InB);
+    int32 a_idx = FindExistingVertex(InA);
+    int32 b_idx = FindExistingVertex(InB);
     if (a_idx == b_idx && a_idx >= 0)
     {
         return false;
@@ -103,14 +103,14 @@ bool FText3DGlyphArrangement::InsertSegment(FVector2d InA, FVector2d InB, double
     FindIntersectingVertices(segAB, a_idx, b_idx, points, InTol);
 
     // insert intersections into existing segments
-    for (int i = 0, N = Hits.Num(); i < N; ++i)
+    for (int32 i = 0, N = Hits.Num(); i < N; ++i)
     {
         FIntersection Intr = Hits[i];
-        int EID = Intr.EID;
+        int32 EID = Intr.EID;
         double t0 = Intr.Intr.Parameter0, t1 = Intr.Intr.Parameter1;
 
         // insert first point at t0
-        int new_eid = -1;
+        int32 new_eid = -1;
         if (Intr.Intr.Type == EIntersectionType::Point || Intr.Intr.Type == EIntersectionType::Segment)
         {
             FIndex2i new_info = SplitSegmentAtDistance(EID, t0, VertexSnapTol);
@@ -152,7 +152,7 @@ bool FText3DGlyphArrangement::InsertSegment(FVector2d InA, FVector2d InB, double
     if (a_idx == -1)
     {
         a_idx = Graph.AppendVertex(InA);
-        PointHash.InsertPointUnsafe(a_idx, InA);
+        PointHash.InsertPoint(a_idx, InA);
     }
     if (b_idx == -1)
     {
@@ -161,7 +161,7 @@ bool FText3DGlyphArrangement::InsertSegment(FVector2d InA, FVector2d InB, double
     if (b_idx == -1)
     {
         b_idx = Graph.AppendVertex(InB);
-        PointHash.InsertPointUnsafe(b_idx, InB);
+        PointHash.InsertPoint(b_idx, InB);
     }
 
     // add start/end to points list. These may be duplicates but we will sort that out after
@@ -172,10 +172,10 @@ bool FText3DGlyphArrangement::InsertSegment(FVector2d InA, FVector2d InB, double
 
     // connect sequential points, as long as they aren't the same point,
     // and the segment doesn't already exist
-    for (int k = 0; k < points.Num() - 1; ++k)
+    for (int32 k = 0; k < points.Num() - 1; ++k)
     {
-        int v0 = points[k].VID;
-        int v1 = points[k + 1].VID;
+        int32 v0 = points[k].VID;
+        int32 v1 = points[k + 1].VID;
         if (v0 == v1)
         {
             continue;
@@ -186,7 +186,7 @@ bool FText3DGlyphArrangement::InsertSegment(FVector2d InA, FVector2d InB, double
             // sanity check; technically this can happen and still be correct but it's more likely an error case
             ensureMsgf(FMath::Abs(points[k].T - points[k + 1].T) >= std::numeric_limits<float>::epsilon(), TEXT("insert_segment: different points have same T??"));
 
-            const int EID = Graph.AppendEdge(v0, v1);
+            const int32 EID = Graph.AppendEdge(v0, v1);
             Directions.Add(EID, v0 < v1);
         }
     }
@@ -194,13 +194,13 @@ bool FText3DGlyphArrangement::InsertSegment(FVector2d InA, FVector2d InB, double
     return true;
 }
 
-FIndex2i FText3DGlyphArrangement::SplitSegmentAtDistance(int InEID, double InDistance, double InTol)
+FIndex2i FText3DGlyphArrangement::SplitSegmentAtDistance(int32 InEID, double InDistance, double InTol)
 {
     FIndex2i ev = Graph.GetEdgeV(InEID);
     FSegment2d seg = FSegment2d(Graph.GetVertex(ev.A), Graph.GetVertex(ev.B));
 
-    int use_vid = -1;
-    int new_eid = -1;
+    int32 use_vid = -1;
+    int32 new_eid = -1;
     if (InDistance < -(seg.Extent - InTol))
     {
         use_vid = ev.A;
@@ -214,7 +214,7 @@ FIndex2i FText3DGlyphArrangement::SplitSegmentAtDistance(int InEID, double InDis
         FVector2d Pt = seg.PointAt(InDistance);
         FDynamicGraph2d::FEdgeSplitInfo splitInfo;
         EMeshResult result;
-        int CrossingVert = FindExistingVertex(Pt);
+        int32 CrossingVert = FindExistingVertex(Pt);
         if (CrossingVert == -1)
         {
             result = Graph.SplitEdge(InEID, splitInfo);
@@ -232,22 +232,22 @@ FIndex2i FText3DGlyphArrangement::SplitSegmentAtDistance(int InEID, double InDis
         if (CrossingVert == -1)
         {	// position + track added vertex
             Graph.SetVertex(use_vid, Pt);
-            PointHash.InsertPointUnsafe(splitInfo.VNew, Pt);
+            PointHash.InsertPoint(splitInfo.VNew, Pt);
         }
     }
     return FIndex2i(use_vid, new_eid);
 }
 
-int FText3DGlyphArrangement::FindExistingVertex(const FVector2d& InPoint) const
+int32 FText3DGlyphArrangement::FindExistingVertex(const FVector2d& InPoint) const
 {
     return FindNearestVertex(InPoint, VertexSnapTol);
 }
 
-int FText3DGlyphArrangement::FindNearestVertex(const FVector2d& InPoint, double InSearchRadius, int InIgnoreVID) const
+int32 FText3DGlyphArrangement::FindNearestVertex(const FVector2d& InPoint, double InSearchRadius, int32 InIgnoreVID) const
 {
-    auto FuncDistSq = [&](int B) { return DistanceSquared(InPoint, Graph.GetVertex(B)); };
-    auto FuncIgnore = [&](int VID) { return VID == InIgnoreVID; };
-    TPair<int, double> found = (InIgnoreVID == -1) ? PointHash.FindNearestInRadius(InPoint, InSearchRadius, FuncDistSq)
+    auto FuncDistSq = [&](int32 B) { return DistanceSquared(InPoint, Graph.GetVertex(B)); };
+    auto FuncIgnore = [&](int32 VID) { return VID == InIgnoreVID; };
+    TPair<int32, double> found = (InIgnoreVID == -1) ? PointHash.FindNearestInRadius(InPoint, InSearchRadius, FuncDistSq)
                                                  : PointHash.FindNearestInRadius(InPoint, InSearchRadius, FuncDistSq, FuncIgnore);
     if (found.Key == PointHash.GetInvalidValue())
     {
@@ -258,18 +258,18 @@ int FText3DGlyphArrangement::FindNearestVertex(const FVector2d& InPoint, double
 
 bool FText3DGlyphArrangement::FindIntersectingEdges(const FVector2d& InA, const FVector2d& InB, TArray<FIntersection>& OutHits, double InTol) const
 {
-    int num_hits = 0;
+    int32 num_hits = 0;
     FVector2d x = FVector2d::Zero(), y = FVector2d::Zero();
     FVector2d EPerp = UE::Geometry::PerpCW(InB - InA);
     UE::Geometry::Normalize(EPerp);
-    for (int EID : Graph.EdgeIndices())
+    for (int32 EID : Graph.EdgeIndices())
     {
         Graph.GetEdgeV(EID, x, y);
         // inlined version of WhichSide with pre-normalized EPerp, to ensure Tolerance is consistent for different edge lengths
         double SignX = EPerp.Dot(x - InA);
         double SignY = EPerp.Dot(y - InA);
-        int SideX = (SignX > InTol ? +1 : (SignX < -InTol ? -1 : 0));
-        int SideY = (SignY > InTol ? +1 : (SignY < -InTol ? -1 : 0));
+        int32 SideX = (SignX > InTol ? +1 : (SignX < -InTol ? -1 : 0));
+        int32 SideY = (SignY > InTol ? +1 : (SignY < -InTol ? -1 : 0));
         if (SideX == SideY && SideX != 0)
         {
             continue; // both pts on same side
@@ -293,9 +293,9 @@ bool FText3DGlyphArrangement::FindIntersectingEdges(const FVector2d& InA, const
 
 bool FText3DGlyphArrangement::FindIntersectingVertices(const FSegment2d &InSegmentAB, int32 InVIDA, int32 InVIDB, TArray<FSegmentPoint>& OutHits, double InTol)
 {
-    int num_hits = 0;
+    int32 num_hits = 0;
 
-    for (int VID : Graph.VertexIndices())
+    for (int32 VID : Graph.VertexIndices())
     {
         if (Graph.GetVtxEdgeCount(VID) > 0 || VID == InVIDA || VID == InVIDB) // if it's an existing edge or on the currently added edge, it's not floating so skip it
         {
diff --git a/Engine/Plugins/VirtualProduction/Text3D/Source/Text3D/Private/GeometryBuilders/Text3DGlyphArrangement.h b/Engine/Plugins/VirtualProduction/Text3D/Source/Text3D/Private/GeometryBuilders/Text3DGlyphArrangement.h
--- a/Engine/Plugins/VirtualProduction/Text3D/Source/Text3D/Private/GeometryBuilders/Text3DGlyphArrangement.h
+++ b/Engine/Plugins/VirtualProduction/Text3D/Source/Text3D/Private/GeometryBuilders/Text3DGlyphArrangement.h
@@ -2,24 +2,25 @@
 
 #pragma once
 
+#include "Text3DPointHashGrid2d.h"
 #include "Curve/DynamicGraph2.h"
-#include "Spatial/PointHashGrid2.h"
 #include "Intersection/IntrSegment2Segment2.h"
 
 /**
  * @brief Modified copy of FArrangement2d, has edge directions
  */
-struct FText3DGlyphArrangement final
+struct FText3DGlyphArrangement
 {
 	using FIndex2i = UE::Geometry::FIndex2i;
 	using FSegment2f = UE::Geometry::FSegment2f;
 	using FSegment2d = UE::Geometry::FSegment2d;
 	using FAxisAlignedBox2f = UE::Geometry::FAxisAlignedBox2f;
-	
+
     UE::Geometry::FDynamicGraph2d Graph;
-    UE::Geometry::TPointHashGrid2d<int> PointHash;
-    TMap<int, bool> Directions; //bool - from A to B
-	const double VertexSnapTol = 0.001;
+    UE::Text3D::FPointHashGrid2d PointHash;
+    TMap<int32, bool> Directions; //bool - from A to B
+
+	static constexpr double VertexSnapTol = 0.001;
 
 	explicit FText3DGlyphArrangement(const FAxisAlignedBox2f& InBoundsHint);
     void Insert(const FSegment2f& Segment);
@@ -39,11 +40,11 @@ protected:
         UE::Geometry::FIntrSegment2Segment2d Intr;
     };
 
-    int InsertPoint(const FVector2d& InPoint, double InTol = 0);
+    int32 InsertPoint(const FVector2d& InPoint, double InTol = 0);
     bool InsertSegment(FVector2d InA, FVector2d InB, double InTol = 0);
     UE::Geometry::FIndex2i SplitSegmentAtDistance(int InEID, double InDistance, double InTol);
-    int FindExistingVertex(const FVector2d& InPoint) const;
-    int FindNearestVertex(const FVector2d& InPoint, double InSearchRadius, int InIgnoreVID = -1) const;
+    int32 FindExistingVertex(const FVector2d& InPoint) const;
+    int32 FindNearestVertex(const FVector2d& InPoint, double InSearchRadius, int InIgnoreVID = -1) const;
     bool FindIntersectingEdges(const FVector2d& InA, const FVector2d& InB, TArray<FIntersection>& OutHits, double InTol = 0) const;
     bool FindIntersectingVertices(const FSegment2d& InSegmentAB, int32 InVIDA, int32 InVIDB, TArray<FSegmentPoint>& OutHits, double InTol = 0);
 };
diff --git a/Engine/Plugins/VirtualProduction/Text3D/Source/Text3D/Private/GeometryBuilders/Text3DGlyphContourList.cpp b/Engine/Plugins/VirtualProduction/Text3D/Source/Text3D/Private/GeometryBuilders/Text3DGlyphContourList.cpp
--- a/Engine/Plugins/VirtualProduction/Text3D/Source/Text3D/Private/GeometryBuilders/Text3DGlyphContourList.cpp
+++ b/Engine/Plugins/VirtualProduction/Text3D/Source/Text3D/Private/GeometryBuilders/Text3DGlyphContourList.cpp
@@ -16,7 +16,7 @@ FText3DGlyphContourList::FText3DGlyphContourList(const FText3DGlyphContourList&
 	}
 }
 
-void FText3DGlyphContourList::Initialize(const TSharedRef<FText3DGlyphData>& Data)
+void FText3DGlyphContourList::Initialize(FText3DGlyphData& Data)
 {
 	for (FText3DGlyphContour& Contour : *this)
 	{
@@ -62,8 +62,8 @@ void FText3DGlyphContourList::Initialize(const TSharedRef<FText3DGlyphData>& Dat
 					Curr->Position = CornerPosition + Curr->TangentX * Offset;
 					Added->Position = CornerPosition - Prev->TangentX * Offset;
 
-					Data->AddVertices(1);
-					const int32 VertexID = Data->AddVertex(Added->Position, { 1.f, 0.f }, { -1.f, 0.f, 0.f });
+					Data.AddVertices(1);
+					const int32 VertexID = Data.AddVertex(Added->Position, { 1.f, 0.f }, { -1.f, 0.f, 0.f });
 
 					Added->PathPrev.Add(VertexID);
 					Added->PathNext.Add(VertexID);
diff --git a/Engine/Plugins/VirtualProduction/Text3D/Source/Text3D/Private/GeometryBuilders/Text3DGlyphContourList.h b/Engine/Plugins/VirtualProduction/Text3D/Source/Text3D/Private/GeometryBuilders/Text3DGlyphContourList.h
--- a/Engine/Plugins/VirtualProduction/Text3D/Source/Text3D/Private/GeometryBuilders/Text3DGlyphContourList.h
+++ b/Engine/Plugins/VirtualProduction/Text3D/Source/Text3D/Private/GeometryBuilders/Text3DGlyphContourList.h
@@ -18,7 +18,7 @@ public:
 	 * Initialize Countours.
 	 * @param Data - Needed to add vertices for split corners.
 	 */
-	void Initialize(const TSharedRef<FText3DGlyphData>& Data);
+	void Initialize(FText3DGlyphData& Data);
 
 	/**
 	 * Create contour.
diff --git a/Engine/Plugins/VirtualProduction/Text3D/Source/Text3D/Private/GeometryBuilders/Text3DGlyphContourNode.h b/Engine/Plugins/VirtualProduction/Text3D/Source/Text3D/Private/GeometryBuilders/Text3DGlyphContourNode.h
--- a/Engine/Plugins/VirtualProduction/Text3D/Source/Text3D/Private/GeometryBuilders/Text3DGlyphContourNode.h
+++ b/Engine/Plugins/VirtualProduction/Text3D/Source/Text3D/Private/GeometryBuilders/Text3DGlyphContourNode.h
@@ -14,22 +14,22 @@ using TText3DGlyphContourNodeShared = TSharedPtr<FText3DGlyphContourNode>;
 
 struct FText3DGlyphContourNode final
 {
-	explicit FText3DGlyphContourNode(const TSharedPtr<FPolygon2f> ContourIn, const bool bCanHaveIntersectionsIn, const bool bClockwiseIn)
-		: Contour(ContourIn)
-		, bCanHaveIntersections(bCanHaveIntersectionsIn)
-		, bClockwise(bClockwiseIn)
+	explicit FText3DGlyphContourNode(const TOptional<FPolygon2f> InContour, const bool bInCanHaveIntersections, const bool bInClockwise)
+		: Contour(InContour)
+		, bCanHaveIntersections(bInCanHaveIntersections)
+		, bClockwise(bInClockwise)
 	{
 	}
 
-	const TSharedPtr<FPolygon2f> Contour;
+	TOptional<FPolygon2f> Contour;
 
 	// Needed for dividing contours with self-intersections: default value is true, false for parts of divided contours
-	const bool bCanHaveIntersections;
+	bool bCanHaveIntersections;
 
 	bool bClockwise;
 
 	// Contours that are inside this contour
-	TArray<TText3DGlyphContourNodeShared> Children;
+	TArray<FText3DGlyphContourNode> Children;
 
 	/** Debug print to display contour data */
 	void Print(int32 InDepth = 0) const
@@ -38,14 +38,14 @@ struct FText3DGlyphContourNode final
 
 		UE_LOG(LogText3D, Log, TEXT("%sContour: Vtx=%d, Area=%f, Bounds=%s, CW=%s"),
 			*Indent,
-			Contour.IsValid() ? Contour->VertexCount() : 0,
-			Contour.IsValid() ? Contour->SignedArea() : 0,
-			Contour.IsValid() ? *Contour->Bounds().Extents().ToString() : TEXT(""),
+			Contour.IsSet() ? Contour->VertexCount() : 0,
+			Contour.IsSet() ? Contour->SignedArea() : 0,
+			Contour.IsSet() ? *Contour->Bounds().Extents().ToString() : TEXT(""),
 			bClockwise ? TEXT("true") : TEXT("false"));
 
-		for (const TSharedPtr<FText3DGlyphContourNode>& Child : Children)
+		for (const FText3DGlyphContourNode& Child : Children)
 		{
-			Child->Print(InDepth + 1);
+			Child.Print(InDepth + 1);
 		}
 	}
 };
diff --git a/Engine/Plugins/VirtualProduction/Text3D/Source/Text3D/Private/GeometryBuilders/Text3DGlyphData.cpp b/Engine/Plugins/VirtualProduction/Text3D/Source/Text3D/Private/GeometryBuilders/Text3DGlyphData.cpp
--- a/Engine/Plugins/VirtualProduction/Text3D/Source/Text3D/Private/GeometryBuilders/Text3DGlyphData.cpp
+++ b/Engine/Plugins/VirtualProduction/Text3D/Source/Text3D/Private/GeometryBuilders/Text3DGlyphData.cpp
@@ -7,8 +7,7 @@
 #include "Math/Vector.h"
 #include "Text3DModule.h"
 
-FText3DGlyphData::FText3DGlyphData(const TSharedPtr<FText3DGlyph>& InGlyph) :
-	Glyph(InGlyph)
+FText3DGlyphData::FText3DGlyphData()
 {
 	CurrentGroup = EText3DGroupType::Front;
 	GroupExpand = 0.0f;
@@ -34,10 +33,9 @@ FText3DGlyphData::FText3DGlyphData(const TSharedPtr<FText3DGlyph>& InGlyph) :
 void FText3DGlyphData::SetCurrentGroup(const EText3DGroupType Type, const float GroupExpandIn)
 {
 	CurrentGroup = Type;
-	check(Glyph.Get());
 
-	FText3DPolygonGroup& Group = Glyph->GetGroups()[static_cast<int32>(Type)];
-	const FMeshDescription& MeshDescription = Glyph->GetMeshDescription();
+	FText3DPolygonGroup& Group = Glyph.GetGroups()[static_cast<int32>(Type)];
+	const FMeshDescription& MeshDescription = Glyph.GetMeshDescription();
 
 	Group.FirstVertex = MeshDescription.Vertices().Num();
 	Group.FirstTriangle = MeshDescription.Triangles().Num();
@@ -62,14 +60,13 @@ void FText3DGlyphData::SetTarget(const float ExtrudeTargetIn, const float Expand
 
 int32 FText3DGlyphData::AddVertices(const int32 Count)
 {
-	check(Glyph.Get());
 	if (Count <= 0)
 	{
 		return 0;
 	}
 
-	FMeshDescription& MeshDescription = Glyph->GetMeshDescription();
-	FStaticMeshAttributes& MeshAttributes = Glyph->GetStaticMeshAttributes();
+	FMeshDescription& MeshDescription = Glyph.GetMeshDescription();
+	FStaticMeshAttributes& MeshAttributes = Glyph.GetStaticMeshAttributes();
 	VertexCountBeforeAdd = MeshDescription.Vertices().Num();
 
 	MeshDescription.ReserveNewVertices(Count);
@@ -98,8 +95,7 @@ int32 FText3DGlyphData::AddVertex(const FVector2D Position, const FVector2D Tang
 
 int32 FText3DGlyphData::AddVertex(const FVector& Position, const FVector& TangentX, const FVector& TangentZ, const FVector2D TextureCoordinates)
 {
-	check(Glyph.Get());
-	FStaticMeshAttributes& StaticMeshAttributes = Glyph->GetStaticMeshAttributes();
+	FStaticMeshAttributes& StaticMeshAttributes = Glyph.GetStaticMeshAttributes();
 	const int32 VertexIndex = VertexCountBeforeAdd + AddVertexIndex++;
 	StaticMeshAttributes.GetVertexPositions()[FVertexID(VertexIndex)] = (FVector3f)Position;
 	const FVertexInstanceID Instance(static_cast<uint32>(VertexIndex));
@@ -113,13 +109,12 @@ int32 FText3DGlyphData::AddVertex(const FVector& Position, const FVector& Tangen
 
 void FText3DGlyphData::AddTriangles(const int32 Count)
 {
-	check(Glyph.Get());
 	if (Count <= 0)
 	{
 		return;
 	}
 
-	FMeshDescription& MeshDescription = Glyph->GetMeshDescription();
+	FMeshDescription& MeshDescription = Glyph.GetMeshDescription();
 	TriangleCountBeforeAdd = MeshDescription.Triangles().Num();
 	MeshDescription.ReserveNewTriangles(Count);
 	AddTriangleIndex = 0;
@@ -127,8 +122,7 @@ void FText3DGlyphData::AddTriangles(const int32 Count)
 
 void FText3DGlyphData::AddTriangle(const int32 A, const int32 B, const int32 C)
 {
-	check(Glyph.Get());
-	Glyph->GetMeshDescription().CreateTriangle(FPolygonGroupID(static_cast<int32>(CurrentGroup)), TArray<FVertexInstanceID>({FVertexInstanceID(A), FVertexInstanceID(B), FVertexInstanceID(C)}));
+	Glyph.GetMeshDescription().CreateTriangle(FPolygonGroupID(static_cast<int32>(CurrentGroup)), TArray<FVertexInstanceID>({FVertexInstanceID(A), FVertexInstanceID(B), FVertexInstanceID(C)}));
 	AddTriangleIndex++;
 }
 
diff --git a/Engine/Plugins/VirtualProduction/Text3D/Source/Text3D/Private/GeometryBuilders/Text3DGlyphData.h b/Engine/Plugins/VirtualProduction/Text3D/Source/Text3D/Private/GeometryBuilders/Text3DGlyphData.h
--- a/Engine/Plugins/VirtualProduction/Text3D/Source/Text3D/Private/GeometryBuilders/Text3DGlyphData.h
+++ b/Engine/Plugins/VirtualProduction/Text3D/Source/Text3D/Private/GeometryBuilders/Text3DGlyphData.h
@@ -14,7 +14,17 @@ public:
 	 * Constructor.
 	 * @param InGlyph - Glyph that is being created.
 	 */
-	explicit FText3DGlyphData(const TSharedPtr<FText3DGlyph>& InGlyph);
+	explicit FText3DGlyphData();
+
+	const FText3DGlyph& GetGlyph() const
+	{
+		return Glyph;
+	}
+
+	FText3DGlyph& GetGlyph()
+	{
+		return Glyph;
+	}
 
 	/**
 	 * Set group that is being created.
@@ -74,7 +84,7 @@ public:
 	void FillEdge(const FText3DGlyphPartPtr& Edge, const bool bSkipLastTriangle, bool bFlipNormals = false);
 
 private:
-	TSharedPtr<FText3DGlyph> Glyph;
+	FText3DGlyph Glyph;
 
 	EText3DGroupType CurrentGroup;
 	float GroupExpand;
diff --git a/Engine/Plugins/VirtualProduction/Text3D/Source/Text3D/Private/GeometryBuilders/Text3DGlyphLoader.cpp b/Engine/Plugins/VirtualProduction/Text3D/Source/Text3D/Private/GeometryBuilders/Text3DGlyphLoader.cpp
--- a/Engine/Plugins/VirtualProduction/Text3D/Source/Text3D/Private/GeometryBuilders/Text3DGlyphLoader.cpp
+++ b/Engine/Plugins/VirtualProduction/Text3D/Source/Text3D/Private/GeometryBuilders/Text3DGlyphLoader.cpp
@@ -10,41 +10,37 @@
 
 using namespace UE::Geometry;
 
-#if WITH_FREETYPE
-FText3DGlyphLoader::FText3DGlyphLoader(const FT_GlyphSlot InGlyph)
-	: Root(MakeShared<FText3DGlyphContourNode>(nullptr, false, true))
-	, EndIndex(-1)
+FText3DGlyphLoader::FText3DGlyphLoader()
+	: EndIndex(-1)
 	, VertexID(0)
 {
-	check(InGlyph);
+}
+
+#if WITH_FREETYPE
+FText3DGlyphContourNode FText3DGlyphLoader::GenerateContourList(const FT_GlyphSlot InGlyph)
+{
+	FText3DGlyphContourNode Root(/*Contour*/{}, /*CanHaveIntersections*/false, /*Clockwise*/true);
 
 	const FT_Outline Outline = InGlyph->outline;
 	const int32 ContourCount = Outline.n_contours;
 
 	for (int32 Index = 0; Index < ContourCount; Index++)
 	{
-		const TSharedPtr<FPolygon2f> Contour = CreateContour(Outline, Index);
-
-		if (Contour.IsValid())
+		const TOptional<FPolygon2f> Contour = CreateContour(Outline, Index);
+		if (Contour.IsSet())
 		{
-			Insert(MakeShared<FText3DGlyphContourNode>(Contour, true, Contour->IsClockwise()), Root);
+			Insert(FText3DGlyphContourNode(Contour, /*CanHaveIntersections*/true, Contour->IsClockwise()), Root);
 		}
 	}
 
-	if (Root->Children.Num() == 0)
+	if (!Root.Children.IsEmpty())
 	{
-		return;
+		FixParityAndDivide(Root, false);
 	}
 
-	FixParityAndDivide(Root, false);
-}
-#endif
-
-TText3DGlyphContourNodeShared FText3DGlyphLoader::GetContourList() const
-{
 	return Root;
 }
-
+#endif
 
 FText3DGlyphLoader::FLine::FLine(const FVector2D PositionIn)
 	: Position(PositionIn)
@@ -291,23 +287,25 @@ FVector2D FText3DGlyphLoader::FCubicCurve::Tangent(const float T)
 }
 
 #if WITH_FREETYPE
-TSharedPtr<FPolygon2f> FText3DGlyphLoader::CreateContour(const FT_Outline Outline, const int32 ContourIndex)
+TOptional<FPolygon2f> FText3DGlyphLoader::CreateContour(const FT_Outline Outline, const int32 ContourIndex)
 {
-	const TSharedPtr<FPolygon2f> Contour = ProcessFreetypeOutline(Outline, ContourIndex);
-	return Contour.IsValid() ? RemoveBadPoints(Contour) : Contour;
+	TOptional<FPolygon2f> Contour = ProcessFreetypeOutline(Outline, ContourIndex);
+	if (Contour.IsSet() && RemoveBadPoints(*Contour))
+	{
+		return Contour;
+	}
+	return {};
 }
 
-TSharedPtr<FPolygon2f> FText3DGlyphLoader::ProcessFreetypeOutline(const FT_Outline Outline, const int32 ContourIndex)
+TOptional<FPolygon2f> FText3DGlyphLoader::ProcessFreetypeOutline(const FT_Outline Outline, const int32 ContourIndex)
 {
-	TSharedPtr<FPolygon2f> Contour;
-
 	const int32 StartIndex = EndIndex + 1;
 	EndIndex = Outline.contours[ContourIndex];
 	const int32 ContourLength = EndIndex - StartIndex + 1;
 
 	if (ContourLength < 3)
 	{
-		return Contour;
+		return {};
 	}
 
 	ProcessedContour.Empty();
@@ -371,7 +369,7 @@ TSharedPtr<FPolygon2f> FText3DGlyphLoader::ProcessFreetypeOutline(const FT_Outli
 
 			if (ContourIsBad(Current))
 			{
-				return Contour;
+				return {};
 			}
 
 			if (TagNext == FT_Curve_Tag_On && (Current - Next).IsNearlyZero())
@@ -389,7 +387,7 @@ TSharedPtr<FPolygon2f> FText3DGlyphLoader::ProcessFreetypeOutline(const FT_Outli
 			{
 				if (ContourIsBad(Prev))
 				{
-					return Contour;
+					return {};
 				}
 
 				A = Prev;
@@ -405,7 +403,7 @@ TSharedPtr<FPolygon2f> FText3DGlyphLoader::ProcessFreetypeOutline(const FT_Outli
 		{
 			if (ContourIsBad(Prev))
 			{
-				return Contour;
+				return {};
 			}
 
 			FCubicCurve(Prev, Current, Next, NextNext).Add(this);
@@ -414,60 +412,50 @@ TSharedPtr<FPolygon2f> FText3DGlyphLoader::ProcessFreetypeOutline(const FT_Outli
 
 	if (ProcessedContour.Num() < 3)
 	{
-		return Contour;
+		return {};
 	}
 
-
-	Contour = MakeShared<FPolygon2f>();
-
-	for (const TDoubleLinkedList<FVector2f>::TDoubleLinkedListNode* Point = ProcessedContour.GetTail(); Point; Point =
-		 Point->GetPrevNode())
+	FPolygon2f Contour;
+	for (const TDoubleLinkedList<FVector2f>::TDoubleLinkedListNode* Point = ProcessedContour.GetTail(); Point; Point = Point->GetPrevNode())
 	{
-		Contour->AppendVertex(Point->GetValue());
+		Contour.AppendVertex(Point->GetValue());
 	}
-
 	return Contour;
 }
 #endif
 
-bool FText3DGlyphLoader::IsInside(const TSharedPtr<const FPolygon2f> ContourA, const TSharedPtr<const FPolygon2f> ContourB) const
+bool FText3DGlyphLoader::IsInside(const FPolygon2f& ContourA, const FPolygon2f& ContourB) const
 {
-	return ContourB->Contains(*ContourA);
+	return ContourB.Contains(ContourA);
 }
 
-void FText3DGlyphLoader::Insert(const TText3DGlyphContourNodeShared NodeA, const TText3DGlyphContourNodeShared NodeB) const
+void FText3DGlyphLoader::Insert(FText3DGlyphContourNode&& NodeA, FText3DGlyphContourNode& NodeB) const
 {
-	TArray<TText3DGlyphContourNodeShared>& ChildrenB = NodeB->Children;
-	const TSharedPtr<const FPolygon2f> ContourA = NodeA->Contour;
-
-	for (int32 ChildBIndex = 0; ChildBIndex < ChildrenB.Num(); ChildBIndex++)
+	for (int32 ChildBIndex = 0; ChildBIndex < NodeB.Children.Num(); ChildBIndex++)
 	{
-		TText3DGlyphContourNodeShared& ChildB = ChildrenB[ChildBIndex];
-		const TSharedPtr<const FPolygon2f> ChildBContour = ChildB->Contour;
+		FText3DGlyphContourNode& ChildB = NodeB.Children[ChildBIndex];
 
-		if (IsInside(ContourA, ChildBContour))
+		if (IsInside(*NodeA.Contour, *ChildB.Contour))
 		{
-			Insert(NodeA, ChildB);
+			Insert(MoveTemp(NodeA), ChildB);
 			return;
 		}
 
-		if (IsInside(ChildBContour, ContourA))
+		if (IsInside(*ChildB.Contour, *NodeA.Contour))
 		{
 			// add ChildBContour to list of contours that are inside ContourA
-			TArray<TText3DGlyphContourNodeShared>& ChildrenA = NodeA->Children;
-			ChildrenA.Add(ChildB);
+			NodeA.Children.Add(ChildB);
 			// replace ChildBContour with ContourA in list it was before
 			ChildB = NodeA;
 
 			// check if other contours in that list are inside ContourA
-			for (int32 ChildBSiblingIndex = ChildrenB.Num() - 1; ChildBSiblingIndex > ChildBIndex; ChildBSiblingIndex--)
+			for (int32 ChildBSiblingIndex = NodeB.Children.Num() - 1; ChildBSiblingIndex > ChildBIndex; ChildBSiblingIndex--)
 			{
-				TText3DGlyphContourNodeShared ChildBSibling = ChildrenB[ChildBSiblingIndex];
-
-				if (IsInside(ChildBSibling->Contour, ContourA))
+				FText3DGlyphContourNode& ChildBSibling = NodeB.Children[ChildBSiblingIndex];
+				if (IsInside(*ChildBSibling.Contour, *NodeA.Contour))
 				{
-					ChildrenA.Add(ChildBSibling);
-					ChildrenB.RemoveAt(ChildBSiblingIndex);
+					NodeA.Children.Add(ChildBSibling);
+					NodeB.Children.RemoveAt(ChildBSiblingIndex);
 				}
 			}
 
@@ -475,56 +463,55 @@ void FText3DGlyphLoader::Insert(const TText3DGlyphContourNodeShared NodeA, const
 		}
 	}
 
-	ChildrenB.Add(NodeA);
+	NodeB.Children.Add(MoveTemp(NodeA));
 }
 
-void FText3DGlyphLoader::FixParityAndDivide(const TText3DGlyphContourNodeShared Node, const bool bClockwiseIn)
+void FText3DGlyphLoader::FixParityAndDivide(FText3DGlyphContourNode& Node, const bool bInClockwise)
 {
-	TArray<TText3DGlyphContourNodeShared>& Children = Node->Children;
-
-	for (int32 ChildIndex = 0; ChildIndex < Children.Num(); ChildIndex++)
+	for (int32 ChildIndex = 0; ChildIndex < Node.Children.Num(); ChildIndex++)
 	{
-		TText3DGlyphContourNodeShared Child = Children[ChildIndex];
-		const TSharedPtr<FPolygon2f> ChildContour = Child->Contour;
+		FText3DGlyphContourNode& Child = Node.Children[ChildIndex];
 
 		// bCanHaveIntersections has default value "true", if it's false then contour was checked for having self-intersections and it's parity was fixed
-		if (!Child->bCanHaveIntersections)
+		if (!Child.bCanHaveIntersections)
 		{
-			FixParityAndDivide(Child, !bClockwiseIn);
+			FixParityAndDivide(Child, !bInClockwise);
 			continue;
 		}
 
 		// Fix parity
-		if (Child->bClockwise != bClockwiseIn)
+		if (Child.bClockwise != bInClockwise)
 		{
-			ChildContour->Reverse();
-			Child->bClockwise = bClockwiseIn;
+			Child.Contour->Reverse();
+			Child.bClockwise = bInClockwise;
 		}
 
-		MakeArrangement(ChildContour);
+		FText3DGlyphArrangement Arrangement = MakeArrangement(*Child.Contour);
 		Junctions.Reset();
 
 		// Find first junction
-		if (!FindJunction())
+		if (!FindJunction(Arrangement))
 		{
 			// If no junction was found, contour doesn't have self-intersections
-			FixParityAndDivide(Child, !bClockwiseIn);
+			FixParityAndDivide(Child, !bInClockwise);
 			continue;
 		}
 
 		// Remove this contour from parent's child list
-		Children.RemoveAt(ChildIndex--);
+		FText3DGlyphContourNode RemovedChildNode = MoveTemp(Child);
+		Node.Children.RemoveAt(ChildIndex--);
+
 		// Reset path
 		DividedContourIDs.Reset();
 		// Parts are stored as children of separate root node
-		TText3DGlyphContourNodeShared RootForDetaching = MakeShared<FText3DGlyphContourNode>(nullptr, false, Node->bClockwise);
+		FText3DGlyphContourNode RootForDetaching({}, /*CanHaveIntersections*/false, Node.bClockwise);
 		// Processed edges are removed from graph to simplify dividing, vertices that don't belong to any edge are removed too, so a copy is needed
-		const TArray<FVector2f> InitialVertices = CopyVertices();
+		const TArray<FVector2f> InitialVertices = CopyVertices(Arrangement);
 
 		bool bAddPoint = true;
 		while (bAddPoint)
 		{
-			AddPoint();
+			AddPoint(Arrangement);
 
 			if (VertexID == -1)
 			{
@@ -545,12 +532,12 @@ void FText3DGlyphLoader::FixParityAndDivide(const TText3DGlyphContourNodeShared
 			}
 
 			// If path contains no junctions and current vertex is not a junction
-			if (Junctions.Num() == 0 && !Arrangement->Graph.IsJunctionVertex(VertexID))
+			if (Junctions.Num() == 0 && !Arrangement.Graph.IsJunctionVertex(VertexID))
 			{
-				if (!FindJunction())
+				if (!FindJunction(Arrangement))
 				{
 					// If there are no junctions in graph, it contains last contour that should be separated
-					if (!FindRegular())
+					if (!FindRegular(Arrangement))
 					{
 						// Graph is empty, done
 						bAddPoint = false;
@@ -558,51 +545,49 @@ void FText3DGlyphLoader::FixParityAndDivide(const TText3DGlyphContourNodeShared
 					else
 					{
 						// Add non-junction vertex as if it was a junction
-						AddPoint(true);
+						AddPoint(Arrangement, true);
 					}
 				}
 			}
 		}
 
 		RemoveUnneededNodes(RootForDetaching);
-		MergeRootForDetaching(Child, RootForDetaching, Node);
+		MergeRootForDetaching(MoveTemp(RemovedChildNode), RootForDetaching, Node);
 	}
 }
 
-bool FText3DGlyphLoader::IsBadNormal(const TSharedPtr<FPolygon2f> Contour, const int32 Point) const
+bool FText3DGlyphLoader::IsBadNormal(const FPolygon2f& InContour, const int32 InPoint) const
 {
 	FVector2f ToPrev;
 	FVector2f ToNext;
-
-	Contour->NeighbourVectors(Point, ToPrev, ToNext, true);
+	InContour.NeighbourVectors(InPoint, ToPrev, ToNext, /*Normalize*/true);
 	return FMath::IsNearlyZero(1.0f - FVector2D::DotProduct(FVector2D(ToPrev), FVector2D(ToNext)));
 }
 
-bool FText3DGlyphLoader::MergedWithNext(const TSharedPtr<FPolygon2f> Contour, const int32 Point) const
+bool FText3DGlyphLoader::MergedWithNext(const FPolygon2f& InContour, const int32 Point) const
 {
 	const int32 Current = Point;
-	const int32 Next = (Point + 1) % Contour->VertexCount();
-
-	return FVector2D((*Contour)[Next] - (*Contour)[Current]).IsNearlyZero();
+	const int32 Next = (Point + 1) % InContour.VertexCount();
+	return FVector2D(InContour[Next] - InContour[Current]).IsNearlyZero();
 }
 
-TSharedPtr<FPolygon2f> FText3DGlyphLoader::RemoveBadPoints(const TSharedPtr<FPolygon2f> Contour) const
+bool FText3DGlyphLoader::RemoveBadPoints(FPolygon2f& InOutContour) const
 {
-	int32 VertexCount = Contour->VertexCount();
+	int32 VertexCount = InOutContour.VertexCount();
 	int32 CheckVertexIndex = 0;
 	int32 ConsecutiveCleanVertices = 0;
 
 	while (ConsecutiveCleanVertices < VertexCount)
 	{
-		if (IsBadNormal(Contour, CheckVertexIndex) || MergedWithNext(Contour, CheckVertexIndex))
+		if (IsBadNormal(InOutContour, CheckVertexIndex) || MergedWithNext(InOutContour, CheckVertexIndex))
 		{
-			Contour->RemoveVertex(CheckVertexIndex);
+			InOutContour.RemoveVertex(CheckVertexIndex);
 
 			// Update count since we removed one
-			VertexCount = Contour->VertexCount();
+			VertexCount = InOutContour.VertexCount();
 			if (VertexCount < 4)
 			{
-				return nullptr;
+				return false;
 			}
 
 			// After removal, the next vertex now occupies CheckVertexIndex
@@ -619,51 +604,44 @@ TSharedPtr<FPolygon2f> FText3DGlyphLoader::RemoveBadPoints(const TSharedPtr<FPol
 			CheckVertexIndex = (CheckVertexIndex + 1) % VertexCount;
 		}
 	}
-
-	return Contour;
+	return true;
 }
 
-void FText3DGlyphLoader::MakeArrangement(const TSharedPtr<const FPolygon2f> Contour)
+FText3DGlyphArrangement FText3DGlyphLoader::MakeArrangement(const FPolygon2f& InContour)
 {
 	FAxisAlignedBox2f Box = FAxisAlignedBox2f::Empty();
-
-	for (const FVector2f& Vertex : Contour->GetVertices())
+	for (const FVector2f& Vertex : InContour.GetVertices())
 	{
 		Box.Contain(Vertex);
 	}
 
-	Arrangement = MakeShared<FText3DGlyphArrangement>(Box);
-
-	for (FSegment2f Edge : Contour->Segments())
+	FText3DGlyphArrangement Arrangement(Box);
+	for (FSegment2f Edge : InContour.Segments())
 	{
-		Arrangement->Insert(Edge);
+		Arrangement.Insert(Edge);
 	}
+	return Arrangement;
 }
 
-bool FText3DGlyphLoader::FindJunction()
+bool FText3DGlyphLoader::FindJunction(const FText3DGlyphArrangement& InArrangement)
 {
-	FDynamicGraph2d& Graph = Arrangement->Graph;
-
-	for (int32 VID = 0; VID < Graph.MaxVertexID(); VID++)
+	for (int32 VID = 0; VID < InArrangement.Graph.MaxVertexID(); VID++)
 	{
-		if (Graph.IsJunctionVertex(VID))
+		if (InArrangement.Graph.IsJunctionVertex(VID))
 		{
 			VertexID = VID;
 			return true;
 		}
 	}
-
 	VertexID = -1;
 	return false;
 }
 
-bool FText3DGlyphLoader::FindRegular()
+bool FText3DGlyphLoader::FindRegular(const FText3DGlyphArrangement& InArrangement)
 {
-	FDynamicGraph2d& Graph = Arrangement->Graph;
-
-	for (int32 VID = 0; VID < Graph.MaxVertexID(); VID++)
+	for (int32 VID = 0; VID < InArrangement.Graph.MaxVertexID(); VID++)
 	{
-		if (Graph.IsRegularVertex(VID))
+		if (InArrangement.Graph.IsRegularVertex(VID))
 		{
 			VertexID = VID;
 			return true;
@@ -674,16 +652,15 @@ bool FText3DGlyphLoader::FindRegular()
 	return false;
 }
 
-TArray<FVector2f> FText3DGlyphLoader::CopyVertices() const
+TArray<FVector2f> FText3DGlyphLoader::CopyVertices(const FText3DGlyphArrangement& InArrangement) const
 {
-	FDynamicGraph2d& Graph = Arrangement->Graph;
 	TArray<FVector2f> Vertices;
-	const int32 MaxVID = Graph.MaxVertexID();
+	const int32 MaxVID = InArrangement.Graph.MaxVertexID();
 	Vertices.Reserve(MaxVID);
 
 	for (int32 VID = 0; VID < MaxVID; VID++)
 	{
-		const FVector2d Vertex = Graph.GetVertex(VID);
+		const FVector2d Vertex = InArrangement.Graph.GetVertex(VID);
 		Vertices.Add({static_cast<float>(Vertex.X), static_cast<float>(Vertex.Y)});
 	}
 
@@ -696,25 +673,19 @@ void FText3DGlyphLoader::DetachBadContour()
 	RemoveDetachedContourFromPath(JunctionIndexInContour);
 }
 
-void FText3DGlyphLoader::DetachFinishedContour(const int32 RepeatedJunctionIndex, const TArray<FVector2f>& InitialVertices,
-										 const TText3DGlyphContourNodeShared RootForDetaching)
+void FText3DGlyphLoader::DetachFinishedContour(const int32 RepeatedJunctionIndex, const TArray<FVector2f>& InitialVertices, FText3DGlyphContourNode& RootForDetaching)
 {
 	const int32 JunctionIndexInContour = FindStartOfDetachedContour(RepeatedJunctionIndex);
 
 	// Create contour from copied vertices
-	TSharedPtr<FPolygon2f> FinishedContour = MakeShared<FPolygon2f>();
-
+	FPolygon2f FinishedContour;
 	for (int32 ID = JunctionIndexInContour; ID < DividedContourIDs.Num(); ID++)
 	{
-		FinishedContour->AppendVertex(InitialVertices[DividedContourIDs[ID]]);
+		FinishedContour.AppendVertex(InitialVertices[DividedContourIDs[ID]]);
 	}
 
 	RemoveDetachedContourFromPath(JunctionIndexInContour);
-
-	const TText3DGlyphContourNodeShared FinishedContourNode = MakeShared<FText3DGlyphContourNode>(FinishedContour,
-																			false,
-																			FinishedContour->IsClockwise());
-	Insert(FinishedContourNode, RootForDetaching);
+	Insert(FText3DGlyphContourNode(FinishedContour, /*CanHaveIntersections*/false, FinishedContour.IsClockwise()), RootForDetaching);
 }
 
 int32 FText3DGlyphLoader::FindStartOfDetachedContour(const int32 RepeatedJunctionIndex)
@@ -731,22 +702,21 @@ void FText3DGlyphLoader::RemoveDetachedContourFromPath(const int32 JunctionIndex
 	DividedContourIDs.SetNum(JunctionIndexInContour);
 }
 
-void FText3DGlyphLoader::AddPoint(const bool bForceJunction)
+void FText3DGlyphLoader::AddPoint(FText3DGlyphArrangement& InArrangement, const bool bForceJunction)
 {
 	DividedContourIDs.Add(VertexID);
-	FDynamicGraph2d& Graph = Arrangement->Graph;
 	int32 OutgoingEdge = 0;
 
-	if (Graph.IsJunctionVertex(VertexID) || bForceJunction)
+	if (InArrangement.Graph.IsJunctionVertex(VertexID) || bForceJunction)
 	{
 		Junctions.Add(VertexID);
 		// Select specific path direction
-		OutgoingEdge = FindOutgoing(Junctions.Last());
+		OutgoingEdge = FindOutgoing(InArrangement, Junctions.Last());
 	}
-	else if (Graph.IsVertex(VertexID))
+	else if (InArrangement.Graph.IsVertex(VertexID))
 	{
 		// This should be the only possible direction
-		OutgoingEdge = *Graph.VtxEdgesItr(VertexID).begin();
+		OutgoingEdge = *InArrangement.Graph.VtxEdgesItr(VertexID).begin();
 	}
 	else
 	{
@@ -756,21 +726,18 @@ void FText3DGlyphLoader::AddPoint(const bool bForceJunction)
 	}
 
 	// Move to next graph node
-	VertexID = GetOtherEdgeVertex(OutgoingEdge, VertexID);
-	Graph.RemoveEdge(OutgoingEdge, true);
+	VertexID = GetOtherEdgeVertex(InArrangement, OutgoingEdge, VertexID);
+	InArrangement.Graph.RemoveEdge(OutgoingEdge, true);
 }
 
-void FText3DGlyphLoader::RemoveUnneededNodes(const TText3DGlyphContourNodeShared Node) const
+void FText3DGlyphLoader::RemoveUnneededNodes(FText3DGlyphContourNode& Node) const
 {
-	TArray<TText3DGlyphContourNodeShared>& Children = Node->Children;
-
-	for (int32 ChildIndex = 0; ChildIndex < Children.Num();)
+	for (int32 ChildIndex = 0; ChildIndex < Node.Children.Num();)
 	{
-		const TText3DGlyphContourNodeShared Child = Children[ChildIndex];
-
-		if (Child->bClockwise == Node->bClockwise)
+		FText3DGlyphContourNode& Child = Node.Children[ChildIndex];
+		if (Child.bClockwise == Node.bClockwise)
 		{
-			Children.RemoveAt(ChildIndex);
+			Node.Children.RemoveAt(ChildIndex);
 		}
 		else
 		{
@@ -780,31 +747,30 @@ void FText3DGlyphLoader::RemoveUnneededNodes(const TText3DGlyphContourNodeShared
 	}
 }
 
-void FText3DGlyphLoader::MergeRootForDetaching(const TText3DGlyphContourNodeShared RemovedNode,
-										 const TText3DGlyphContourNodeShared RootForDetaching,
-										 const TText3DGlyphContourNodeShared RemovedNodeParent) const
+void FText3DGlyphLoader::MergeRootForDetaching(FText3DGlyphContourNode&& RemovedNode, FText3DGlyphContourNode& RootForDetaching, FText3DGlyphContourNode& RemovedNodeParent) const
 {
-	for (const TText3DGlyphContourNodeShared& Child : RemovedNode->Children)
+	for (FText3DGlyphContourNode& Child : RemovedNode.Children)
 	{
-		Insert(Child, RootForDetaching);
+		Insert(MoveTemp(Child), RootForDetaching);
 	}
+	RemovedNode.Children.Reset();
 
-	for (const TText3DGlyphContourNodeShared& Child : RootForDetaching->Children)
+	for (FText3DGlyphContourNode& Child : RootForDetaching.Children)
 	{
-		RemovedNodeParent->Children.Add(Child);
+		RemovedNodeParent.Children.Add(Child);
 	}
 }
 
-bool FText3DGlyphLoader::IsOutgoing(const int32 Junction, const int32 EID) const
+bool FText3DGlyphLoader::IsOutgoing(const FText3DGlyphArrangement& InArrangement, const int32 Junction, const int32 EID) const
 {
-	return (Arrangement->Graph.GetEdgeV(EID).A == Junction) == Arrangement->Directions[EID];
+	return (InArrangement.Graph.GetEdgeV(EID).A == Junction) == InArrangement.Directions[EID];
 }
 
-int32 FText3DGlyphLoader::FindOutgoing(const int32 Junction) const
+int32 FText3DGlyphLoader::FindOutgoing(const FText3DGlyphArrangement& InArrangement, const int32 Junction) const
 {
 	// Get list of all edges that have this vertex
 	TArray<int32> Edges;
-	SortedVtxEdges(Junction, Edges);
+	SortedVtxEdges(InArrangement, Junction, Edges);
 
 	int32 Current = 0;
 	int32 Next = Edges[0];
@@ -815,7 +781,7 @@ int32 FText3DGlyphLoader::FindOutgoing(const int32 Junction) const
 		Next = Edges[(EdgeIndex + 1) % Edges.Num()];
 
 		// Needed direction is the one that is outgoing and it's clockwise neighbour is not
-		if (!IsOutgoing(Junction, Current) && IsOutgoing(Junction, Next))
+		if (!IsOutgoing(InArrangement, Junction, Current) && IsOutgoing(InArrangement, Junction, Next))
 		{
 			return Next;
 		}
@@ -824,28 +790,27 @@ int32 FText3DGlyphLoader::FindOutgoing(const int32 Junction) const
 	return Edges[0];
 }
 
-int32 FText3DGlyphLoader::GetOtherEdgeVertex(const int32 EID, const int32 VID) const
+int32 FText3DGlyphLoader::GetOtherEdgeVertex(const FText3DGlyphArrangement& InArrangement, const int32 EID, const int32 VID) const
 {
-	const FIndex2i Edge = Arrangement->Graph.GetEdgeV(EID);
+	const FIndex2i Edge = InArrangement.Graph.GetEdgeV(EID);
 	return (Edge.A == VID) ? Edge.B : ((Edge.B == VID) ? Edge.A : FDynamicGraph::InvalidID);
 }
 
-void FText3DGlyphLoader::SortedVtxEdges(const int32 VID, TArray<int32>& Sorted) const
+void FText3DGlyphLoader::SortedVtxEdges(const FText3DGlyphArrangement& InArrangement, const int32 VID, TArray<int32>& Sorted) const
 {
-	FDynamicGraph2d& Graph = Arrangement->Graph;
-	Sorted.Reserve(Graph.GetVtxEdgeCount(VID));
+	Sorted.Reserve(InArrangement.Graph.GetVtxEdgeCount(VID));
 
-	for (int32 EID : Graph.VtxEdgesItr(VID))
+	for (int32 EID : InArrangement.Graph.VtxEdgesItr(VID))
 	{
 		Sorted.Add(EID);
 	}
 
-	const FVector2d V = Graph.GetVertex(VID);
+	const FVector2d V = InArrangement.Graph.GetVertex(VID);
 	Algo::SortBy(Sorted,
-				 [&](int32 EID)
-				 {
-					 const int32 NbrVID = GetOtherEdgeVertex(EID, VID);
-					 const FVector2d D = Graph.GetVertex(NbrVID) - V;
-					 return TMathUtil<double>::Atan2Positive(D.Y, D.X);
-				 });
+		[&](int32 EID)
+		{
+			 const int32 NbrVID = GetOtherEdgeVertex(InArrangement, EID, VID);
+			 const FVector2d D = InArrangement.Graph.GetVertex(NbrVID) - V;
+			 return TMathUtil<double>::Atan2Positive(D.Y, D.X);
+		});
 }
diff --git a/Engine/Plugins/VirtualProduction/Text3D/Source/Text3D/Private/GeometryBuilders/Text3DGlyphLoader.h b/Engine/Plugins/VirtualProduction/Text3D/Source/Text3D/Private/GeometryBuilders/Text3DGlyphLoader.h
--- a/Engine/Plugins/VirtualProduction/Text3D/Source/Text3D/Private/GeometryBuilders/Text3DGlyphLoader.h
+++ b/Engine/Plugins/VirtualProduction/Text3D/Source/Text3D/Private/GeometryBuilders/Text3DGlyphLoader.h
@@ -9,18 +9,19 @@
 #include "Templates/SharedPointer.h"
 #include "Text3DModule.h"
 
+struct FText3DGlyphArrangement;
 using UE::Geometry::FPolygon2f;
 
 /** Creates glyph geometry data from a font glyphs */
 class FText3DGlyphLoader final
 {
 public:
+	explicit FText3DGlyphLoader();
+
 #if WITH_FREETYPE
-	explicit FText3DGlyphLoader(const FT_GlyphSlot InGlyph);
+	FText3DGlyphContourNode GenerateContourList(const FT_GlyphSlot InGlyph);
 #endif
 
-	TText3DGlyphContourNodeShared GetContourList() const;
-
 private:
 	class FLine final
 	{
@@ -123,12 +124,9 @@ private:
 		FVector2D Tangent(const float t) override;
 	};
 
-
-	TText3DGlyphContourNodeShared Root;
 	int32 EndIndex;
 	TDoubleLinkedList<FVector2f> ProcessedContour;
 	FVector2D FirstPosition;
-	TSharedPtr<struct FText3DGlyphArrangement> Arrangement;
 	// List of junctions in path
 	TArray<int32> Junctions;
 	// Current vertex that should be added to path
@@ -137,8 +135,8 @@ private:
 	TArray<int32> DividedContourIDs;
 
 #if WITH_FREETYPE
-	TSharedPtr<FPolygon2f> CreateContour(const FT_Outline Outline, const int32 ContourIndex);
-	TSharedPtr<FPolygon2f> ProcessFreetypeOutline(const FT_Outline Outline, const int32 ContourIndex);
+	TOptional<FPolygon2f> CreateContour(const FT_Outline Outline, const int32 ContourIndex);
+	TOptional<FPolygon2f> ProcessFreetypeOutline(const FT_Outline Outline, const int32 ContourIndex);
 #endif
 
 	/**
@@ -146,38 +144,38 @@ private:
 	 * @param ContourA
 	 * @param ContourB
 	 */
-	bool IsInside(const TSharedPtr<const FPolygon2f> ContourA, const TSharedPtr<const FPolygon2f> ContourB) const;
+	bool IsInside(const FPolygon2f& ContourA, const FPolygon2f& ContourB) const;
 
 	/**
 	 * Insert NodeA inside NodeB.
 	 * @param NodeA
 	 * @param NodeB
 	 */
-	void Insert(const TText3DGlyphContourNodeShared NodeA, const TText3DGlyphContourNodeShared NodeB) const;
+	void Insert(FText3DGlyphContourNode&& NodeA, FText3DGlyphContourNode& NodeB) const;
 	/**
 	 * Reverse contour if it's initial parity differs from the one it should have, find self-intersections and, if found, divide to contours without self-intersections.
 	 * @param Node - Function is called recursively to fix all contours inside Node->Contour.
-	 * @param ClockwiseIn - The parity that contours listed in Node->Nodes should have.
+	 * @param bInClockwise - The parity that contours listed in Node->Nodes should have.
 	 */
-	void FixParityAndDivide(const TText3DGlyphContourNodeShared Node, const bool bClockwiseIn);
+	void FixParityAndDivide(FText3DGlyphContourNode& Node, const bool bInClockwise);
 
-	TSharedPtr<FPolygon2f> RemoveBadPoints(const TSharedPtr<FPolygon2f> Contour) const;
+	bool RemoveBadPoints(FPolygon2f& InOutContour) const;
 
-	bool IsBadNormal(const TSharedPtr<FPolygon2f> Contour, const int32 Point) const;
-	bool MergedWithNext(const TSharedPtr<FPolygon2f> Contour, const int32 Point) const;
+	bool IsBadNormal(const FPolygon2f& InContour, const int32 InPoint) const;
+	bool MergedWithNext(const FPolygon2f& InContour, const int32 Point) const;
 
-	void MakeArrangement(const TSharedPtr<const FPolygon2f> Contour);
+	FText3DGlyphArrangement MakeArrangement(const FPolygon2f& InContour);
 	/**
 	 * Find vertex that belongs to more then 2 edges and write it's ID to VertexID member.
 	 * @return Was such vertex found?
 	 */
-	bool FindJunction();
+	bool FindJunction(const FText3DGlyphArrangement& InArrangement);
 	/**
 	 * Find vertex that belongs to 2 edges and write it's ID to VertexID member.
 	 * @return Was such vertex found?
 	 */
-	bool FindRegular();
-	TArray<FVector2f> CopyVertices() const;
+	bool FindRegular(const FText3DGlyphArrangement& InArrangement);
+	TArray<FVector2f> CopyVertices(const FText3DGlyphArrangement& InArrangement) const;
 
 	/**
 	 * Separate bad contour (the one that is not a closed loop) from path.
@@ -189,7 +187,7 @@ private:
 	 * @param InitialVertices - Previously copied graph vertices.
 	 * @param RootForDetaching - Separate root node for detaching.
 	 */
-	void DetachFinishedContour(const int32 RepeatedJunctionIndex, const TArray<FVector2f>& InitialVertices, const TText3DGlyphContourNodeShared RootForDetaching);
+	void DetachFinishedContour(const int32 RepeatedJunctionIndex, const TArray<FVector2f>& InitialVertices, FText3DGlyphContourNode& RootForDetaching);
 
 	/**
 	 * Find start of contour that is detached, remove it from list of junctions.
@@ -205,33 +203,36 @@ private:
 
 	/**
 	 * Add vertex to path
+	 * @param InArrangement the arrangement to query
 	 * @param bForceJunction - Needed for separating last loop, it has no junctions so any other vertex has to be treated as one.
 	 */
-	void AddPoint(const bool bForceJunction = false);
+	void AddPoint(FText3DGlyphArrangement& InArrangement, const bool bForceJunction = false);
+
 	/**
 	 * Check what nodes have same parity as their parents and remove them, they cover the area that already is covered by parent.
 	 * @param Node - Processed node.
 	 */
-	void RemoveUnneededNodes(const TText3DGlyphContourNodeShared Node) const;
+	void RemoveUnneededNodes(FText3DGlyphContourNode& Node) const;
+
 	/**
 	 * Merge separate root to regular one.
 	 * @param RemovedNode - Contour that was divided.
 	 * @param RootForDetaching - Separate root that stores result of division.
 	 * @param RemovedNodeParent - Parent of contour that was divided.
 	 */
-	void MergeRootForDetaching(const TText3DGlyphContourNodeShared RemovedNode, const TText3DGlyphContourNodeShared RootForDetaching, const TText3DGlyphContourNodeShared RemovedNodeParent) const;
+	void MergeRootForDetaching(FText3DGlyphContourNode&& RemovedNode, FText3DGlyphContourNode& RootForDetaching, FText3DGlyphContourNode& RemovedNodeParent) const;
 
 	/**
 	 * Choose next path direction in junction vertex.
 	 * @param Junction - Vertex ID of junction.
 	 * @return Edge ID of chosen direction.
 	 */
-	int32 FindOutgoing(const int32 Junction) const;
-	bool IsOutgoing(const int32 Junction, const int32 EID) const;
+	int32 FindOutgoing(const FText3DGlyphArrangement& InArrangement, const int32 Junction) const;
+	bool IsOutgoing(const FText3DGlyphArrangement& InArrangement, const int32 Junction, const int32 EID) const;
 
 	// Copied from FDynamicGraph2d because it's private there
-	int32 GetOtherEdgeVertex(const int32 EID, const int32 VID) const;
+	int32 GetOtherEdgeVertex(const FText3DGlyphArrangement& InArrangement, const int32 EID, const int32 VID) const;
 
 	// Fixed version of FDynamicGraph2d::SortedVtxEdges
-	void SortedVtxEdges(const int32 VID, TArray<int32>& Sorted) const;
+	void SortedVtxEdges(const FText3DGlyphArrangement& InArrangement, const int32 VID, TArray<int32>& Sorted) const;
 };
diff --git a/Engine/Plugins/VirtualProduction/Text3D/Source/Text3D/Private/GeometryBuilders/Text3DGlyphMeshBuilder.cpp b/Engine/Plugins/VirtualProduction/Text3D/Source/Text3D/Private/GeometryBuilders/Text3DGlyphMeshBuilder.cpp
--- a/Engine/Plugins/VirtualProduction/Text3D/Source/Text3D/Private/GeometryBuilders/Text3DGlyphMeshBuilder.cpp
+++ b/Engine/Plugins/VirtualProduction/Text3D/Source/Text3D/Private/GeometryBuilders/Text3DGlyphMeshBuilder.cpp
@@ -11,22 +11,16 @@
 
 using namespace UE::Geometry;
 
-FText3DGlyphMeshBuilder::FText3DGlyphMeshBuilder() :
-	Glyph(MakeShared<FText3DGlyph>()),
-	Data(MakeShared<FText3DGlyphData>(Glyph))
+void FText3DGlyphMeshBuilder::CreateMeshes(const FText3DGlyphContourNode& InRoot, const float InExtrude, const float InBevel, const EText3DBevelType InBevelType, const int32 InBevelSegments, const bool bInOutline, const float InOutlineExpand, const EText3DOutlineType InOutlineType)
 {
-}
-
-void FText3DGlyphMeshBuilder::CreateMeshes(const TText3DGlyphContourNodeShared& Root, const float InExtrude, const float InBevel, const EText3DBevelType InBevelType, const int32 InBevelSegments, const bool bInOutline, const float InOutlineExpand, const EText3DOutlineType InOutlineType)
-{
-	CreateFrontMesh(Root, bInOutline, InOutlineExpand, InOutlineType);
+	CreateFrontMesh(InRoot, bInOutline, InOutlineExpand, InOutlineType);
 
-	if (Contours->Num() == 0)
+	if (Contours.Num() == 0)
 	{
 		return;
 	}
 
-	const bool bFlipNormals = FMath::Sign(Data->GetPlannedExpand()) < 0;
+	const bool bFlipNormals = FMath::Sign(Data.GetPlannedExpand()) < 0;
 	const float BevelLocal = bInOutline ? 0.f : InBevel;
 
 	if (!bInOutline)
@@ -43,12 +37,12 @@ void FText3DGlyphMeshBuilder::SetFrontAndBevelTextureCoordinates(const float Bev
 	int32 GroupIndex = static_cast<int32>(GroupType);
 
 	FBox2f Box;
-	TText3DGroupList& Groups = Glyph->GetGroups();
+	TText3DGroupList& Groups = Data.GetGlyph().GetGroups();
 
 	const int32 FirstVertex = Groups[GroupIndex].FirstVertex;
 	const int32 LastVertex = Groups[GroupIndex + 1].FirstVertex;
 
-	TVertexAttributesConstRef<FVector3f> Positions = Glyph->GetStaticMeshAttributes().GetVertexPositions();
+	TVertexAttributesConstRef<FVector3f> Positions = Data.GetGlyph().GetStaticMeshAttributes().GetVertexPositions();
 
 	const FVector3f& FirstPosition = Positions[FVertexID(FirstVertex)];
 	const FVector2f PositionFlat = { FirstPosition.Y, FirstPosition.Z };
@@ -67,7 +61,7 @@ void FText3DGlyphMeshBuilder::SetFrontAndBevelTextureCoordinates(const float Bev
 		Box.Max.Y = FMath::Max(Box.Max.Y, Position.Z);
 	}
 
-	FStaticMeshAttributes& StaticMeshAttributes = Glyph->GetStaticMeshAttributes();
+	FStaticMeshAttributes& StaticMeshAttributes = Data.GetGlyph().GetStaticMeshAttributes();
 	TVertexAttributesRef<FVector3f> VertexPositions = StaticMeshAttributes.GetVertexPositions();
 	TVertexInstanceAttributesRef<FVector2f> VertexInstanceUVs = StaticMeshAttributes.GetVertexInstanceUVs();
 
@@ -104,14 +98,14 @@ void FText3DGlyphMeshBuilder::MirrorGroups(const float Extrude)
 
 void FText3DGlyphMeshBuilder::MovePivot(const FVector& InNewPivot)
 {
-	const TVertexAttributesRef<FVector3f> VertexPositions = Glyph->GetStaticMeshAttributes().GetVertexPositions();
+	const TVertexAttributesRef<FVector3f> VertexPositions = Data.GetGlyph().GetStaticMeshAttributes().GetVertexPositions();
 
 	// Remove empty space around glyph first
 	float MinX = TNumericLimits<float>::Max();
 	float MaxX = TNumericLimits<float>::Lowest();
 	float MinY = TNumericLimits<float>::Max();
 	float MaxY = TNumericLimits<float>::Lowest();
-	for (const FVertexID VertexID : Glyph->GetMeshDescription().Vertices().GetElementIDs())
+	for (const FVertexID VertexID : Data.GetGlyph().GetMeshDescription().Vertices().GetElementIDs())
 	{
 		const FVector3f Position = VertexPositions[VertexID];
 		MinX = FMath::Min(MinX, Position.X);
@@ -127,7 +121,7 @@ void FText3DGlyphMeshBuilder::MovePivot(const FVector& InNewPivot)
 	// Apply pivot offset with adjusted values
 	const FVector3f PivotOffset = FVector3f(InNewPivot) * FVector3f(MaxX - MinX, MaxY - MinY, 0);
 
-	for (const FVertexID VertexID : Glyph->GetMeshDescription().Vertices().GetElementIDs())
+	for (const FVertexID VertexID : Data.GetGlyph().GetMeshDescription().Vertices().GetElementIDs())
 	{
 		FVector3f CurrentPosition = VertexPositions[VertexID];
 		CurrentPosition += PivotOffset;
@@ -135,16 +129,16 @@ void FText3DGlyphMeshBuilder::MovePivot(const FVector& InNewPivot)
 	}
 }
 
-void FText3DGlyphMeshBuilder::BuildMesh(FText3DCachedMesh& InMesh, UMaterialInterface* InDefaultMaterial)
+void FText3DGlyphMeshBuilder::BuildMesh(FText3DCachedMesh& InMesh, EText3DMeshType InMeshType, UMaterialInterface* InDefaultMaterial)
 {
-	Glyph->Build(InMesh, InDefaultMaterial);
+	Data.GetGlyph().Build(InMesh, InMeshType, InDefaultMaterial);
 	InMesh.MeshBounds = GetMeshBounds();
 	InMesh.MeshOffset = GetMeshOffset();
 }
 
 FBox FText3DGlyphMeshBuilder::GetMeshBounds() const
 {
-	return Glyph->GetMeshDescription().ComputeBoundingBox();
+	return Data.GetGlyph().GetMeshDescription().ComputeBoundingBox();
 }
 
 FVector FText3DGlyphMeshBuilder::GetMeshOffset() const
@@ -152,19 +146,19 @@ FVector FText3DGlyphMeshBuilder::GetMeshOffset() const
 	return MeshOffset;
 }
 
-void FText3DGlyphMeshBuilder::CreateFrontMesh(const TText3DGlyphContourNodeShared& Root, const bool bOutline, const float OutlineExpand, const EText3DOutlineType InOutlineType)
+void FText3DGlyphMeshBuilder::CreateFrontMesh(const FText3DGlyphContourNode& Root, const bool bOutline, const float OutlineExpand, const EText3DOutlineType InOutlineType)
 {
 	int32 VertexCount = 0;
 	AddToVertexCount(Root, VertexCount);
 
-	Data->SetCurrentGroup(EText3DGroupType::Front, bOutline ? OutlineExpand : 0.f);
-	Data->SetTarget(0.f, 0.f);
-	Contours = MakeShared<FText3DGlyphContourList>();
+	Data.SetCurrentGroup(EText3DGroupType::Front, bOutline ? OutlineExpand : 0.f);
+	Data.SetTarget(0.f, 0.f);
+	Contours.Reset();
 
-	int32 VertexIndex = Data->AddVertices(VertexCount);
+	int32 VertexIndex = Data.AddVertices(VertexCount);
 	TriangulateAndConvert(Root, VertexIndex, bOutline, InOutlineType);
 
-	Contours->Initialize(Data);
+	Contours.Initialize(Data);
 
 	if (bOutline)
 	{
@@ -174,7 +168,7 @@ void FText3DGlyphMeshBuilder::CreateFrontMesh(const TText3DGlyphContourNodeShare
 
 void FText3DGlyphMeshBuilder::CreateBevelMesh(const float Bevel, const EText3DBevelType Type, const int32 BevelSegments)
 {
-	Data->SetCurrentGroup(EText3DGroupType::Bevel, Bevel);
+	Data.SetCurrentGroup(EText3DGroupType::Bevel, Bevel);
 
 	if (FMath::IsNearlyZero(Bevel))
 	{
@@ -253,17 +247,17 @@ void FText3DGlyphMeshBuilder::CreateExtrudeMesh(float Extrude, float Bevel, cons
 		Extrude -= Bevel * 2.0f;
 	}
 
-	Data->SetCurrentGroup(EText3DGroupType::Extrude, 0.f);
+	Data.SetCurrentGroup(EText3DGroupType::Extrude, 0.f);
 
 	const FVector2D Normal(1.f, 0.f);
-	Data->PrepareSegment(Extrude, 0.f, Normal, Normal);
+	Data.PrepareSegment(Extrude, 0.f, Normal, Normal);
 
-	Contours->Reset();
+	Contours.Reset();
 
 
 	TArray<float> TextureCoordinateVs;
 
-	for (FText3DGlyphContour& Contour : *Contours)
+	for (FText3DGlyphContour& Contour : Contours)
 	{
 		// Compute TexCoord.V-s for each point
 		TextureCoordinateVs.Reset(Contour.Num() - 1);
@@ -292,7 +286,7 @@ void FText3DGlyphMeshBuilder::CreateExtrudeMesh(float Extrude, float Bevel, cons
 		}
 
 		// Duplicate contour
-		Data->SetTarget(0.f, 0.f);
+		Data.SetTarget(0.f, 0.f);
 		const bool bFirstSmooth = First->bSmooth;
 		// It's set to sharp because we need 2 vertices with TexCoord.Y values 0 and 1 (for smooth points only one vertex is added)
 		First->bSmooth = false;
@@ -324,9 +318,8 @@ void FText3DGlyphMeshBuilder::CreateExtrudeMesh(float Extrude, float Bevel, cons
 			ExpandPoint(Point, {0.f, 1.f - TextureCoordinateVs[Index++ - 1]});
 		}
 
-
 		// Add extruded vertices
-		Data->SetTarget(Data->GetPlannedExtrude(), Data->GetPlannedExpand());
+		Data.SetTarget(Data.GetPlannedExtrude(), Data.GetPlannedExpand());
 
 		// Similarly to duplicating vertices, first point is processed separately
 		{
@@ -355,14 +348,14 @@ void FText3DGlyphMeshBuilder::CreateExtrudeMesh(float Extrude, float Bevel, cons
 
 		for (const FText3DGlyphPartPtr& Edge : Contour)
 		{
-			Data->FillEdge(Edge, false, bFlipNormals);
+			Data.FillEdge(Edge, false, bFlipNormals);
 		}
 	}
 }
 
 void FText3DGlyphMeshBuilder::MirrorGroup(const EText3DGroupType TypeIn, const EText3DGroupType TypeOut, const float Extrude)
 {
-	TText3DGroupList& Groups = Glyph->GetGroups();
+	TText3DGroupList& Groups = Data.GetGlyph().GetGroups();
 
 	const FText3DPolygonGroup GroupIn = Groups[static_cast<int32>(TypeIn)];
 	const FText3DPolygonGroup GroupNext = Groups[static_cast<int32>(TypeIn) + 1];
@@ -370,13 +363,13 @@ void FText3DGlyphMeshBuilder::MirrorGroup(const EText3DGroupType TypeIn, const E
 	const int32 VerticesInNum = GroupNext.FirstVertex - GroupIn.FirstVertex;
 	const int32 TrianglesInNum = GroupNext.FirstTriangle - GroupIn.FirstTriangle;
 
-	FMeshDescription& MeshDescription = Glyph->GetMeshDescription();
+	FMeshDescription& MeshDescription = Data.GetGlyph().GetMeshDescription();
 	const int32 TotalVerticesNum = MeshDescription.Vertices().Num();
 
-	Data->SetCurrentGroup(TypeOut, 0.f);
-	Data->AddVertices(VerticesInNum);
+	Data.SetCurrentGroup(TypeOut, 0.f);
+	Data.AddVertices(VerticesInNum);
 
-	FStaticMeshAttributes& StaticMeshAttributes = Glyph->GetStaticMeshAttributes();
+	FStaticMeshAttributes& StaticMeshAttributes = Data.GetGlyph().GetStaticMeshAttributes();
 	TVertexAttributesRef<FVector3f> VertexPositions = StaticMeshAttributes.GetVertexPositions();
 	TVertexInstanceAttributesRef<FVector3f> VertexNormals = StaticMeshAttributes.GetVertexInstanceNormals();
 	TVertexInstanceAttributesRef<FVector3f> VertexTangents = StaticMeshAttributes.GetVertexInstanceTangents();
@@ -391,10 +384,10 @@ void FText3DGlyphMeshBuilder::MirrorGroup(const EText3DGroupType TypeIn, const E
 		const FVector Normal = (FVector)VertexNormals[InstanceID];
 		const FVector Tangent = (FVector)VertexTangents[InstanceID];
 
-		Data->AddVertex({ Extrude - Position.X, Position.Y, Position.Z }, { -Tangent.X, Tangent.Y, Tangent.Z }, { -Normal.X, Normal.Y, Normal.Z }, FVector2D(VertexUVs[InstanceID]));
+		Data.AddVertex({ Extrude - Position.X, Position.Y, Position.Z }, { -Tangent.X, Tangent.Y, Tangent.Z }, { -Normal.X, Normal.Y, Normal.Z }, FVector2D(VertexUVs[InstanceID]));
 	}
 
-	Data->AddTriangles(TrianglesInNum);
+	Data.AddTriangles(TrianglesInNum);
 
 	for (int32 TriangleIndex = 0; TriangleIndex < TrianglesInNum; TriangleIndex++)
 	{
@@ -404,35 +397,33 @@ void FText3DGlyphMeshBuilder::MirrorGroup(const EText3DGroupType TypeIn, const E
 		uint32 Instance0 = static_cast<uint32>(TotalVerticesNum + VertexInstanceIDs[0].GetValue() - GroupIn.FirstVertex);
 		uint32 Instance2 = static_cast<uint32>(TotalVerticesNum + VertexInstanceIDs[2].GetValue() - GroupIn.FirstVertex);
 		uint32 Instance1 = static_cast<uint32>(TotalVerticesNum + VertexInstanceIDs[1].GetValue() - GroupIn.FirstVertex);
-		Data->AddTriangle(Instance0, Instance2, Instance1);
+		Data.AddTriangle(Instance0, Instance2, Instance1);
 	}
 }
 
-void FText3DGlyphMeshBuilder::AddToVertexCount(const TText3DGlyphContourNodeShared& Node, int32& OutVertexCount)
+void FText3DGlyphMeshBuilder::AddToVertexCount(const FText3DGlyphContourNode& Node, int32& OutVertexCount)
 {
-	for (const TText3DGlyphContourNodeShared& Child : Node->Children)
+	for (const FText3DGlyphContourNode& Child : Node.Children)
 	{
-		OutVertexCount += Child->Contour->VertexCount();
+		OutVertexCount += Child.Contour->VertexCount();
 		AddToVertexCount(Child, OutVertexCount);
 	}
 }
 
-void FText3DGlyphMeshBuilder::TriangulateAndConvert(const TText3DGlyphContourNodeShared& Node, int32& OutVertexIndex, const bool bOutline, const EText3DOutlineType InOutlineType)
+void FText3DGlyphMeshBuilder::TriangulateAndConvert(const FText3DGlyphContourNode& Node, int32& OutVertexIndex, const bool bOutline, const EText3DOutlineType InOutlineType)
 {
 	// If this is solid region
-	if (!Node->bClockwise)
+	if (!Node.bClockwise)
 	{
 		int32 VertexCount = 0;
 		UE::Geometry::FConstrainedDelaunay2f Triangulation;
 		Triangulation.FillRule = UE::Geometry::FConstrainedDelaunay2f::EFillRule::Positive;
 
-		const TSharedPtr<FText3DGlyphContourList> ContoursLocal = Contours;
-		const TSharedRef<FText3DGlyphData> DataLocal = Data;
-		auto ProcessContour = [ContoursLocal, DataLocal, &VertexCount, &Triangulation, bOutline, InOutlineType](const TText3DGlyphContourNodeShared NodeIn)
+		auto ProcessContour = [&Contours=Contours, &Data=Data, &VertexCount, &Triangulation, bOutline, InOutlineType](const FText3DGlyphContourNode& NodeIn)
 		{
 			// Create contour in old format
-			FText3DGlyphContour& Contour = ContoursLocal->Add();
-			const FPolygon2f& Polygon = *NodeIn->Contour;
+			FText3DGlyphContour& Contour = Contours.Add();
+			const FPolygon2f& Polygon = *NodeIn.Contour;
 
 			for (const FVector2f& Vertex : Polygon.GetVertices())
 			{
@@ -442,7 +433,7 @@ void FText3DGlyphMeshBuilder::TriangulateAndConvert(const TText3DGlyphContourNod
 				Point->Position = FVector2D(Vertex);
 
 				// Add point to mesh
-				const int32 VertexID = DataLocal->AddVertex(Point->Position, { 1.f, 0.f }, { -1.f, 0.f, 0.f });
+				const int32 VertexID = Data.AddVertex(Point->Position, { 1.f, 0.f }, { -1.f, 0.f, 0.f });
 
 				Point->PathPrev.Add(VertexID);
 				Point->PathNext.Add(VertexID);
@@ -453,7 +444,7 @@ void FText3DGlyphMeshBuilder::TriangulateAndConvert(const TText3DGlyphContourNod
 			// Add contour to triangulation
 			if (!bOutline || InOutlineType != EText3DOutlineType::Stroke)
 			{
-				Triangulation.Add(Polygon, NodeIn->bClockwise);
+				Triangulation.Add(Polygon, NodeIn.bClockwise);
 			}
 		};
 
@@ -461,7 +452,7 @@ void FText3DGlyphMeshBuilder::TriangulateAndConvert(const TText3DGlyphContourNod
 		ProcessContour(Node);
 
 		// Holes
-		for (const TText3DGlyphContourNodeShared& Child : Node->Children)
+		for (const FText3DGlyphContourNode& Child : Node.Children)
 		{
 			ProcessContour(Child);
 		}
@@ -470,11 +461,11 @@ void FText3DGlyphMeshBuilder::TriangulateAndConvert(const TText3DGlyphContourNod
 		{
 			Triangulation.Triangulate();
 			const TArray<FIndex3i>& Triangles = Triangulation.Triangles;
-			Data->AddTriangles(Triangles.Num());
+			Data.AddTriangles(Triangles.Num());
 
 			for (const FIndex3i& Triangle : Triangles)
 			{
-				Data->AddTriangle(OutVertexIndex + Triangle.A, OutVertexIndex + Triangle.C, OutVertexIndex + Triangle.B);
+				Data.AddTriangle(OutVertexIndex + Triangle.A, OutVertexIndex + Triangle.C, OutVertexIndex + Triangle.B);
 			}
 		}
 
@@ -482,7 +473,7 @@ void FText3DGlyphMeshBuilder::TriangulateAndConvert(const TText3DGlyphContourNod
 	}
 
 	// Continue with children
-	for (const TText3DGlyphContourNodeShared& Child : Node->Children)
+	for (const FText3DGlyphContourNode& Child : Node.Children)
 	{
 		TriangulateAndConvert(Child, OutVertexIndex, bOutline, InOutlineType);
 	}
@@ -490,10 +481,9 @@ void FText3DGlyphMeshBuilder::TriangulateAndConvert(const TText3DGlyphContourNod
 
 void FText3DGlyphMeshBuilder::CreateOutlineMesh(float InOutlineExpand)
 {
-	Data->SetCurrentGroup(EText3DGroupType::Bevel, InOutlineExpand);
-
-	FText3DGlyphContourList InitialContours = *Contours;
+	Data.SetCurrentGroup(EText3DGroupType::Bevel, InOutlineExpand);
 
+	FText3DGlyphContourList InitialContours = Contours;
 	for (FText3DGlyphContour& Contour : InitialContours)
 	{
 		Algo::Reverse(Contour);
@@ -520,12 +510,12 @@ void FText3DGlyphMeshBuilder::CreateOutlineMesh(float InOutlineExpand)
 	const FVector2D Normal = {0.f, -1.f};
 	BevelLinear(0.f, InOutlineExpand, Normal, Normal, false);
 
-	Contours->Reset();
+	Contours.Reset();
 
 	TDoubleLinkedList<FText3DGlyphContour>::TDoubleLinkedListNode* Node = InitialContours.GetHead();
 	while (Node)
 	{
-		Contours->AddTail(Node);
+		Contours.AddTail(Node);
 		InitialContours.RemoveNode(Node, false);
 
 		Node = InitialContours.GetHead();
@@ -615,8 +605,8 @@ void FText3DGlyphMeshBuilder::BevelWithSteps(const float Bevel, const int32 Step
 
 void FText3DGlyphMeshBuilder::BevelLinear(const float Extrude, const float Expand, FVector2D NormalStart, FVector2D NormalEnd, const bool bSmooth)
 {
-	Data->PrepareSegment(Extrude, Expand, NormalStart, NormalEnd);
-	Contours->Reset();
+	Data.PrepareSegment(Extrude, Expand, NormalStart, NormalEnd);
+	Contours.Reset();
 
 	if (!bSmooth)
 	{
@@ -625,14 +615,14 @@ void FText3DGlyphMeshBuilder::BevelLinear(const float Extrude, const float Expan
 
 	BevelPartsWithoutIntersectingNormals();
 
-	Data->IncreaseDoneExtrude();
+	Data.IncreaseDoneExtrude();
 }
 
 void FText3DGlyphMeshBuilder::DuplicateContourVertices()
 {
-	Data->SetTarget(0.f, 0.f);
+	Data.SetTarget(0.f, 0.f);
 
-	for (FText3DGlyphContour& Contour : *Contours)
+	for (FText3DGlyphContour& Contour : Contours)
 	{
 		for (const FText3DGlyphPartPtr& Point : Contour)
 		{
@@ -645,11 +635,11 @@ void FText3DGlyphMeshBuilder::DuplicateContourVertices()
 
 void FText3DGlyphMeshBuilder::BevelPartsWithoutIntersectingNormals()
 {
-	Data->SetTarget(Data->GetPlannedExtrude(), Data->GetPlannedExpand());
-	const float MaxExpand = Data->GetPlannedExpand();
+	Data.SetTarget(Data.GetPlannedExtrude(), Data.GetPlannedExpand());
+	const float MaxExpand = Data.GetPlannedExpand();
 
-	const bool bFlipNormals = FMath::Sign(Data->GetPlannedExpand()) < 0;
-	for (FText3DGlyphContour& Contour : *Contours)
+	const bool bFlipNormals = FMath::Sign(Data.GetPlannedExpand()) < 0;
+	for (FText3DGlyphContour& Contour : Contours)
 	{
 		for (const FText3DGlyphPartPtr& Point : Contour)
 		{
@@ -666,7 +656,7 @@ void FText3DGlyphMeshBuilder::BevelPartsWithoutIntersectingNormals()
 
 		for (const FText3DGlyphPartPtr& Edge : Contour)
 		{
-			Data->FillEdge(Edge, false, bFlipNormals);
+			Data.FillEdge(Edge, false, bFlipNormals);
 		}
 	}
 }
@@ -692,10 +682,10 @@ void FText3DGlyphMeshBuilder::ExpandPoint(const FText3DGlyphPartPtr& Point, cons
 	}
 }
 
-void FText3DGlyphMeshBuilder::ExpandPointWithoutAddingVertices(const FText3DGlyphPartPtr& Point) const
+void FText3DGlyphMeshBuilder::ExpandPointWithoutAddingVertices(const FText3DGlyphPartPtr& Point)
 {
-	Point->Position = Data->Expanded(Point);
-	const int32 FirstAdded = Data->AddVertices(Point->bSmooth ? 1 : 2);
+	Point->Position = Data.Expanded(Point);
+	const int32 FirstAdded = Data.AddVertices(Point->bSmooth ? 1 : 2);
 
 	Point->PathPrev.Add(FirstAdded);
 	Point->PathNext.Add(Point->bSmooth ? FirstAdded : FirstAdded + 1);
@@ -706,10 +696,10 @@ void FText3DGlyphMeshBuilder::AddVertexSmooth(const FText3DGlyphPartConstPtr& Po
 	const FText3DGlyphPartConstPtr Curr = Point;
 	const FText3DGlyphPartConstPtr Prev = Point->Prev;
 
-	Data->AddVertex(Point, (Prev->TangentX + Curr->TangentX).GetSafeNormal(), (Data->ComputeTangentZ(Prev, Point->DoneExpand) + Data->ComputeTangentZ(Curr, Point->DoneExpand)).GetSafeNormal(), TextureCoordinates);
+	Data.AddVertex(Point, (Prev->TangentX + Curr->TangentX).GetSafeNormal(), (Data.ComputeTangentZ(Prev, Point->DoneExpand) + Data.ComputeTangentZ(Curr, Point->DoneExpand)).GetSafeNormal(), TextureCoordinates);
 }
 
 void FText3DGlyphMeshBuilder::AddVertexSharp(const FText3DGlyphPartConstPtr& Point, const FText3DGlyphPartConstPtr& Edge, const FVector2D TextureCoordinates)
 {
-	Data->AddVertex(Point, Edge->TangentX, Data->ComputeTangentZ(Edge, Point->DoneExpand).GetSafeNormal(), TextureCoordinates);
+	Data.AddVertex(Point, Edge->TangentX, Data.ComputeTangentZ(Edge, Point->DoneExpand).GetSafeNormal(), TextureCoordinates);
 }
diff --git a/Engine/Plugins/VirtualProduction/Text3D/Source/Text3D/Private/GeometryBuilders/Text3DGlyphMeshBuilder.h b/Engine/Plugins/VirtualProduction/Text3D/Source/Text3D/Private/GeometryBuilders/Text3DGlyphMeshBuilder.h
--- a/Engine/Plugins/VirtualProduction/Text3D/Source/Text3D/Private/GeometryBuilders/Text3DGlyphMeshBuilder.h
+++ b/Engine/Plugins/VirtualProduction/Text3D/Source/Text3D/Private/GeometryBuilders/Text3DGlyphMeshBuilder.h
@@ -4,10 +4,10 @@
 
 #include "GeometryBuilders/Text3DGlyphContourNode.h"
 #include "GeometryBuilders/Text3DGlyphPart.h"
+#include "Text3DGlyphContourList.h"
+#include "Text3DGlyphData.h"
 #include "Text3DTypes.h"
 
-class FText3DGlyphContourList;
-class FText3DGlyphData;
 class FText3DGlyph;
 class UMaterial;
 class UMaterialInterface;
@@ -18,11 +18,9 @@ struct FText3DCachedMesh;
 class FText3DGlyphMeshBuilder final
 {
 public:
-	FText3DGlyphMeshBuilder();
-
 	/**
 	 * CreateMeshes
-	 * @param Root - Tree of contours.	 
+	 * @param InRoot - Tree of contours.	 
 	 * @param InExtrude - Orthogonal (to front cap) offset value.
 	 * @param InBevel - Bevel value (bevel happens before extrude).
 	 * @param InBevelType - Defines shape of beveled part.
@@ -31,7 +29,7 @@ public:
 	 * @param InOutlineExpand - Offsets the outline by the specified amount.
 	 * @param InOutlineType - Defines the outline behavior.
 	 */
-	void CreateMeshes(const TText3DGlyphContourNodeShared& Root, const float InExtrude, const float InBevel, const EText3DBevelType InBevelType, const int32 InBevelSegments, const bool bInOutline, const float InOutlineExpand, const EText3DOutlineType InOutlineType);
+	void CreateMeshes(const FText3DGlyphContourNode& InRoot, const float InExtrude, const float InBevel, const EText3DBevelType InBevelType, const int32 InBevelSegments, const bool bInOutline, const float InOutlineExpand, const EText3DOutlineType InOutlineType);
 	void SetFrontAndBevelTextureCoordinates(const float Bevel);
 	void MirrorGroups(const float Extrude);
 
@@ -39,7 +37,7 @@ public:
 	void MovePivot(const FVector& InNewPivot);
 
 	/** Creates the actual mesh */
-	void BuildMesh(FText3DCachedMesh& InMesh, UMaterialInterface* InDefaultMaterial);
+	void BuildMesh(FText3DCachedMesh& InMesh, EText3DMeshType InMeshType, UMaterialInterface* InDefaultMaterial);
 
 	/** Get actual glyph bounds once built */
 	FBox GetMeshBounds() const;
@@ -49,9 +47,8 @@ public:
 
 private:
 	FVector MeshOffset = FVector::ZeroVector;
-	TSharedPtr<FText3DGlyph> Glyph;
-	TSharedRef<FText3DGlyphData> Data;
-	TSharedPtr<FText3DGlyphContourList> Contours;
+	FText3DGlyphData Data;
+	FText3DGlyphContourList Contours;
 
 	/**
 	 * Create 'Front' part of glyph.
@@ -60,7 +57,7 @@ private:
 	 * @param OutlineExpand - Offsets the outline by the specified amount.
 	 * @param InOutlineType - Defines the outline behavior.
 	 */
-	void CreateFrontMesh(const TText3DGlyphContourNodeShared& Root, const bool bOutline, const float OutlineExpand, const EText3DOutlineType InOutlineType);
+	void CreateFrontMesh(const FText3DGlyphContourNode& Root, const bool bOutline, const float OutlineExpand, const EText3DOutlineType InOutlineType);
 
 	/**
 	 * Create 'Bevel' part of glyph (actually half of it, will be mirrored later).
@@ -95,7 +92,7 @@ private:
 	 * @param Node - Tree node.
 	 * @param OutVertexCount - Pointer to counter.
 	 */
-	void AddToVertexCount(const TText3DGlyphContourNodeShared& Node, int32& OutVertexCount);
+	void AddToVertexCount(const FText3DGlyphContourNode& Node, int32& OutVertexCount);
 
 	/**
 	 * Triangulate solid region with holes and convert contours to old format.
@@ -103,7 +100,7 @@ private:
 	 * @param OutVertexIndex - Pointer to last used vertex index.
 	 * @param bOutline - Front face has outline and is not filled.
 	 */
-	void TriangulateAndConvert(const TText3DGlyphContourNodeShared& Node, int32& OutVertexIndex, const bool bOutline, const EText3DOutlineType InOutlineType);
+	void TriangulateAndConvert(const FText3DGlyphContourNode& Node, int32& OutVertexIndex, const bool bOutline, const EText3DOutlineType InOutlineType);
 
 	void BevelLinearWithSegments(const float Extrude, const float Expand, const int32 BevelSegments, const FVector2D Normal);
 	void BevelCurve(const float Angle, const int32 BevelSegments, TFunction<FVector2D(const float CurrentCos, const float CurrentSin, const float NextCos, const float Next)> ComputeOffset);
@@ -146,7 +143,7 @@ private:
 	 * Common code for expanding, vertices are added uninitialized.
 	 * @param Point - Expanded point.
 	 */
-	void ExpandPointWithoutAddingVertices(const FText3DGlyphPartPtr& Point) const;
+	void ExpandPointWithoutAddingVertices(const FText3DGlyphPartPtr& Point);
 
 	/**
 	 * Add vertex for smooth point.
diff --git a/Engine/Plugins/VirtualProduction/Text3D/Source/Text3D/Private/GeometryBuilders/Text3DPointHashGrid2d.cpp b/Engine/Plugins/VirtualProduction/Text3D/Source/Text3D/Private/GeometryBuilders/Text3DPointHashGrid2d.cpp
new file mode 100644
--- /dev/null
+++ b/Engine/Plugins/VirtualProduction/Text3D/Source/Text3D/Private/GeometryBuilders/Text3DPointHashGrid2d.cpp
@@ -0,0 +1,66 @@
+// Copyright Epic Games, Inc. All Rights Reserved.
+
+#include "Text3DPointHashGrid2d.h"
+
+namespace UE::Text3D
+{
+
+FPointHashGrid2d::FPointHashGrid2d(double InCellSize, int32 InInvalidValue)
+	: Indexer(InCellSize), InvalidValue(InInvalidValue)
+{
+}
+
+int32 FPointHashGrid2d::GetInvalidValue() const
+{
+	return InvalidValue;
+}
+
+void FPointHashGrid2d::InsertPoint(int32 InValue, const FVector2D& InPosition)
+{
+	UE::Geometry::FVector2i Index = Indexer.ToGrid(InPosition);
+	Hash.Add(Index, InValue);
+}
+
+TPair<int32, double> FPointHashGrid2d::FindNearestInRadius(const FVector2D& InQueryPoint, double InRadius, TFunctionRef<double(int32)> InDistanceSqFunc, TFunctionRef<bool(int32)> InIgnoreFunc) const
+{
+	if (Hash.IsEmpty())
+	{
+		return TPair<int32, double>(GetInvalidValue(), TNumericLimits<double>::Max());
+	}
+
+	const UE::Geometry::FVector2i MinIndex = Indexer.ToGrid(InQueryPoint - InRadius * FVector2D::One());
+	const UE::Geometry::FVector2i MaxIndex = Indexer.ToGrid(InQueryPoint + InRadius * FVector2D::One());
+	const double RadiusSquared = InRadius * InRadius;
+
+	int32 Nearest = GetInvalidValue();
+	double MinDistanceSq = TNumericLimits<double>::Max();
+
+	TArray<int32> Values;
+
+	for (int32 IndexY = MinIndex.Y; IndexY <= MaxIndex.Y; ++IndexY) 
+	{
+		for (int32 IndexX = MinIndex.X; IndexX <= MaxIndex.X; ++IndexX) 
+		{
+			UE::Geometry::FVector2i Index(IndexX, IndexY);
+			Values.Reset();
+
+			Hash.MultiFind(Index, Values);
+			for (int32 Value : Values) 
+			{
+				if (InIgnoreFunc(Value))
+				{
+					continue;
+				}
+				const double DistanceSquared = InDistanceSqFunc(Value);
+				if (DistanceSquared < RadiusSquared && DistanceSquared < MinDistanceSq)
+				{
+					Nearest = Value;
+					MinDistanceSq = DistanceSquared;
+				}
+			}
+		}
+	}
+	return TPair<int32, double>(Nearest, MinDistanceSq);
+}
+
+} // UE::Text3D
diff --git a/Engine/Plugins/VirtualProduction/Text3D/Source/Text3D/Private/GeometryBuilders/Text3DPointHashGrid2d.h b/Engine/Plugins/VirtualProduction/Text3D/Source/Text3D/Private/GeometryBuilders/Text3DPointHashGrid2d.h
new file mode 100644
--- /dev/null
+++ b/Engine/Plugins/VirtualProduction/Text3D/Source/Text3D/Private/GeometryBuilders/Text3DPointHashGrid2d.h
@@ -0,0 +1,54 @@
+// Copyright Epic Games, Inc. All Rights Reserved.
+
+#pragma once
+
+#include "IntVectorTypes.h"
+#include "Containers/Map.h"
+#include "Util/GridIndexing2.h"
+
+namespace UE::Text3D
+{
+
+/**
+ * This is a simplified re-implementation of TPointHashGrid2 
+ * but without the FCriticalSection to allow for copying and because it was unnecessary in our context.
+ */
+class FPointHashGrid2d
+{
+public:
+	/**
+	 * Construct 2D hash grid
+	 * @param InCellSize size of grid cells
+	 * @param InInvalidValue this value will be returned by queries if no valid result is found (e.g. bounded-distance query)
+	 */
+	FPointHashGrid2d(double InCellSize, int32 InInvalidValue);
+
+	/** Invalid grid value */
+	int32 GetInvalidValue() const;
+
+	/**
+	 * Insert at given position.
+	 * @param InValue the point/value to insert
+	 * @param InPosition the position associated with this value
+	 */
+	void InsertPoint(int32 InValue, const FVector2D& InPosition);
+
+	/**
+	 * Find nearest point in grid, within a given sphere
+	 * @param InQueryPoint the center of the query sphere
+	 * @param InRadius the radius of the query sphere
+	 * @param InDistanceSqFunc Function you provide which measures the distance between QueryPoint and a Value
+	 * @param InIgnoreFunc optional Function you may provide which will result in a Value being ignored if IgnoreFunc(Value) returns true
+	 * @return the found pair (Value,DistanceFunc(Value)), or (InvalidValue,MaxDouble) if not found
+	 */
+	TPair<int32, double> FindNearestInRadius(const FVector2D& InQueryPoint, double InRadius
+		, TFunctionRef<double(int32)> InDistanceSqFunc
+		, TFunctionRef<bool(int32)> InIgnoreFunc = [](const int32 data) { return false; }) const;
+
+private:
+	TMultiMap<UE::Geometry::FVector2i, int32> Hash;
+	UE::Geometry::FScaleGridIndexer2d Indexer;
+	int32 InvalidValue;
+};
+
+} // UE::Text3D
diff --git a/Engine/Plugins/VirtualProduction/Text3D/Source/Text3D/Private/Renderers/DynamicMesh/Text3DDynamicMeshRenderer.cpp b/Engine/Plugins/VirtualProduction/Text3D/Source/Text3D/Private/Renderers/DynamicMesh/Text3DDynamicMeshRenderer.cpp
--- a/Engine/Plugins/VirtualProduction/Text3D/Source/Text3D/Private/Renderers/DynamicMesh/Text3DDynamicMeshRenderer.cpp
+++ b/Engine/Plugins/VirtualProduction/Text3D/Source/Text3D/Private/Renderers/DynamicMesh/Text3DDynamicMeshRenderer.cpp
@@ -210,6 +210,11 @@ void UText3DDynamicMeshRenderer::OnDestroy()
 	}
 }
 
+EText3DMeshType UText3DDynamicMeshRenderer::GetMeshType() const
+{
+	return EText3DMeshType::Dynamic;
+}
+
 FName UText3DDynamicMeshRenderer::GetFriendlyName() const
 {
 	static const FName Name(TEXT("DynamicMeshRenderer"));
diff --git a/Engine/Plugins/VirtualProduction/Text3D/Source/Text3D/Private/Renderers/DynamicMesh/Text3DDynamicMeshRenderer.h b/Engine/Plugins/VirtualProduction/Text3D/Source/Text3D/Private/Renderers/DynamicMesh/Text3DDynamicMeshRenderer.h
--- a/Engine/Plugins/VirtualProduction/Text3D/Source/Text3D/Private/Renderers/DynamicMesh/Text3DDynamicMeshRenderer.h
+++ b/Engine/Plugins/VirtualProduction/Text3D/Source/Text3D/Private/Renderers/DynamicMesh/Text3DDynamicMeshRenderer.h
@@ -60,6 +60,7 @@ protected:
 	virtual void OnUpdate(const UE::Text3D::Renderer::FUpdateParameters& InParameters) override;
 	virtual void OnClear() override;
 	virtual void OnDestroy() override;
+	virtual EText3DMeshType GetMeshType() const override;
 	virtual FName GetFriendlyName() const override;
 	virtual FBox OnCalculateBounds() const override;
 #if WITH_EDITOR
diff --git a/Engine/Plugins/VirtualProduction/Text3D/Source/Text3D/Private/Renderers/StaticMeshes/Text3DStaticMeshesRenderer.cpp b/Engine/Plugins/VirtualProduction/Text3D/Source/Text3D/Private/Renderers/StaticMeshes/Text3DStaticMeshesRenderer.cpp
--- a/Engine/Plugins/VirtualProduction/Text3D/Source/Text3D/Private/Renderers/StaticMeshes/Text3DStaticMeshesRenderer.cpp
+++ b/Engine/Plugins/VirtualProduction/Text3D/Source/Text3D/Private/Renderers/StaticMeshes/Text3DStaticMeshesRenderer.cpp
@@ -253,6 +253,11 @@ void UText3DStaticMeshesRenderer::OnDestroy()
 	}
 }
 
+EText3DMeshType UText3DStaticMeshesRenderer::GetMeshType() const
+{
+	return EText3DMeshType::Static;
+}
+
 FName UText3DStaticMeshesRenderer::GetFriendlyName() const
 {
 	static const FName Name(TEXT("StaticMeshesRenderer"));
@@ -329,8 +334,8 @@ void UText3DStaticMeshesRenderer::AllocateComponents(int32 InCount)
 
 			if (MeshComponent)
 			{
-				CharacterMeshes.Add(MeshComponent);
 				Actor->AddOwnedComponent(MeshComponent);
+				CharacterMeshes.Add(MeshComponent);
 				MeshComponent->RegisterComponent();
 				MeshComponent->AttachToComponent(TextRoot, FAttachmentTransformRules::SnapToTargetIncludingScale);
 			}
diff --git a/Engine/Plugins/VirtualProduction/Text3D/Source/Text3D/Private/Renderers/StaticMeshes/Text3DStaticMeshesRenderer.h b/Engine/Plugins/VirtualProduction/Text3D/Source/Text3D/Private/Renderers/StaticMeshes/Text3DStaticMeshesRenderer.h
--- a/Engine/Plugins/VirtualProduction/Text3D/Source/Text3D/Private/Renderers/StaticMeshes/Text3DStaticMeshesRenderer.h
+++ b/Engine/Plugins/VirtualProduction/Text3D/Source/Text3D/Private/Renderers/StaticMeshes/Text3DStaticMeshesRenderer.h
@@ -50,6 +50,7 @@ protected:
 	virtual void OnUpdate(const UE::Text3D::Renderer::FUpdateParameters& InParameters) override;
 	virtual void OnClear() override;
 	virtual void OnDestroy() override;
+	virtual EText3DMeshType GetMeshType() const override;
 	virtual FName GetFriendlyName() const override;
 	virtual FBox OnCalculateBounds() const override;
 #if WITH_EDITOR
@@ -58,7 +59,7 @@ protected:
 #endif
 	//~ End UText3DRendererBase
 
-	void AllocateComponents(int32 InCount);
+	void AllocateComponents(int32 InNewCharacterCount);
 
 #if WITH_EDITOR
 	void SetDebugMode(bool bEnabled);
diff --git a/Engine/Plugins/VirtualProduction/Text3D/Source/Text3D/Private/Subsystems/Text3DEngineSubsystem.cpp b/Engine/Plugins/VirtualProduction/Text3D/Source/Text3D/Private/Subsystems/Text3DEngineSubsystem.cpp
--- a/Engine/Plugins/VirtualProduction/Text3D/Source/Text3D/Private/Subsystems/Text3DEngineSubsystem.cpp
+++ b/Engine/Plugins/VirtualProduction/Text3D/Source/Text3D/Private/Subsystems/Text3DEngineSubsystem.cpp
@@ -162,6 +162,22 @@ FText3DFontFaceCache* UText3DEngineSubsystem::FindOrAddCachedFontFace(const TSha
 	return &FontFaceCache;
 }
 
+FText3DFontFaceCache* UText3DEngineSubsystem::FindCachedFontFace(const TSharedPtr<FFreeTypeFace>& InFontFace)
+{
+	if (!InFontFace.IsValid() || !InFontFace->IsFaceValid())
+	{
+		return nullptr;
+	}
+
+	const uint32 FontFaceHash = FText3DFontFaceCache::GetFontFaceHash(InFontFace);
+	if (FontFaceHash == 0)
+	{
+		return nullptr;
+	}
+
+	return CachedFontFaces.Find(FontFaceHash);
+}
+
 uint32 FText3DFontFaceCache::GetFontFaceHash(const TSharedPtr<FFreeTypeFace>& InFontFace)
 {
 	if (const FT_Face Face = InFontFace->GetFace())
@@ -190,84 +206,93 @@ void FText3DFontFaceCache::UpdateFontFace(const TSharedPtr<FFreeTypeFace>& InFon
 }
 #endif
 
-uint32 FText3DFontFaceCache::GetGlyphMeshHash(uint32 InGlyphIndex, const FGlyphMeshParameters& InParameters)
-{
-	return HashCombine(GetTypeHash(InParameters), GetTypeHash(InGlyphIndex));
-}
-
-FText3DCachedMesh* FText3DFontFaceCache::FindOrBuildGlyphMesh(uint32 InGlyphIndex, const FGlyphMeshParameters& InParameters)
+TOptional<FText3DGlyphContourNode> FText3DFontFaceCache::PrepareCachedMesh(uint32 InGlyphIndex, const FGlyphMeshParameters& InParameters)
 {
 	const uint32 HashParameters = GetGlyphMeshHash(InGlyphIndex, InParameters);
 	if (FText3DCachedMesh* CachedMesh = GlyphMeshes.Find(HashParameters))
 	{
-		return CachedMesh;
-	}
-
-	const TText3DGlyphContourNodeShared Root = LoadGlyphContours(InGlyphIndex);
-	if (!Root || Root->Children.IsEmpty())
-	{
-		return nullptr;
+		return {};
 	}
 
-	UText3DEngineSubsystem* TextSubsystem = UText3DEngineSubsystem::Get();
-	if (!TextSubsystem)
+	TOptional<FText3DGlyphContourNode> GlyphContourRootNode = LoadGlyphContours(InGlyphIndex);
+	if (!GlyphContourRootNode.IsSet() || GlyphContourRootNode->Children.IsEmpty())
 	{
-		return nullptr;
+		return {};
 	}
 
-	FText3DCachedMesh& CachedMesh = GlyphMeshes.FindOrAdd(HashParameters);
+	FText3DCachedMesh& CachedMesh = GlyphMeshes.Add(HashParameters);
 
 #if WITH_FREETYPE
+	// Font face is already valid if a valid root node was produced
 	if (const FT_Face Face = FontFace->GetFace())
 	{
 		CachedMesh.FontFaceGlyphSize = FVector2D(Face->size->metrics.x_ppem, Face->size->metrics.y_ppem);
 	}
 #endif
 
-	const FString MeshName = FString::Printf(TEXT("Text3D_Char_%u"), HashParameters);
-	const FName StaticMeshName = MakeUniqueObjectName(TextSubsystem, UStaticMesh::StaticClass(), FName(*MeshName));
-	CachedMesh.StaticMesh = NewObject<UStaticMesh>(TextSubsystem, StaticMeshName, RF_Transient | RF_Public);
+	return GlyphContourRootNode;
+}
 
-	const FName DynamicMeshName = MakeUniqueObjectName(TextSubsystem, UDynamicMesh::StaticClass(), FName(*MeshName));
-	CachedMesh.DynamicMesh = NewObject<UDynamicMesh>(TextSubsystem, DynamicMeshName, RF_Transient | RF_Public);
+uint32 FText3DFontFaceCache::GetGlyphMeshHash(uint32 InGlyphIndex, const FGlyphMeshParameters& InParameters)
+{
+	return HashCombine(GetTypeHash(InParameters), GetTypeHash(InGlyphIndex));
+}
+
+void FText3DFontFaceCache::BuildGlyphMesh(const FBuildGlyphMeshSharedDesc& InSharedParams, const FBuildGlyphMeshInstanceDesc& InInstanceParams)
+{
+	const FGlyphMeshParameters& GlyphMeshParameters = *InSharedParams.GlyphMeshParameters;
+
+	const uint32 HashParameters = GetGlyphMeshHash(InInstanceParams.GlyphIndex, GlyphMeshParameters);
+	FText3DCachedMesh& CachedMesh = GlyphMeshes.FindChecked(HashParameters);
+
+	{
+		auto OnMeshAllocated = 
+			[bIsInGameThread = IsInGameThread()](TNotNull<UObject*> InMesh)
+			{
+				if (!bIsInGameThread)
+				{
+					// OK to clear the async flags for the objects immediately because they have already been assigned directly to a UPROPERTY.
+					InMesh->AtomicallyClearInternalFlags(EInternalObjectFlags::Async);
+				}
+			};
+
+		FGCScopeGuard ScopeGuard;
+
+		CachedMesh.StaticMesh = NewObject<UStaticMesh>(InSharedParams.TextSubsystem, NAME_None, RF_Transient | RF_Public);
+		OnMeshAllocated(CachedMesh.StaticMesh);
+
+		if (InSharedParams.GlyphMeshParameters->MeshType == EText3DMeshType::Dynamic)
+		{
+			CachedMesh.DynamicMesh = NewObject<UDynamicMesh>(InSharedParams.TextSubsystem, NAME_None, RF_Transient | RF_Public);
+			OnMeshAllocated(CachedMesh.DynamicMesh);
+		}
+	}
 
 	FText3DGlyphMeshBuilder MeshCreator;
 	MeshCreator.CreateMeshes(
-		Root,
-		InParameters.Extrude,
-		InParameters.Bevel,
-		InParameters.BevelType,
-		InParameters.BevelSegments,
-		InParameters.bOutline,
-		InParameters.OutlineExpand,
-		InParameters.OutlineType
+		*InInstanceParams.RootNode,
+		GlyphMeshParameters.Extrude,
+		GlyphMeshParameters.Bevel,
+		GlyphMeshParameters.BevelType,
+		GlyphMeshParameters.BevelSegments,
+		GlyphMeshParameters.bOutline,
+		GlyphMeshParameters.OutlineExpand,
+		GlyphMeshParameters.OutlineType
 	);
-	MeshCreator.SetFrontAndBevelTextureCoordinates(InParameters.Bevel);
-	MeshCreator.MirrorGroups(InParameters.Extrude);
-	MeshCreator.MovePivot(InParameters.PivotOffset);
-	MeshCreator.BuildMesh(CachedMesh, GetDefault<UText3DProjectSettings>()->GetDefaultMaterial());
-
-	return &CachedMesh;
+	MeshCreator.SetFrontAndBevelTextureCoordinates(GlyphMeshParameters.Bevel);
+	MeshCreator.MirrorGroups(GlyphMeshParameters.Extrude);
+	MeshCreator.MovePivot(GlyphMeshParameters.PivotOffset);
+	MeshCreator.BuildMesh(CachedMesh, GlyphMeshParameters.MeshType, InSharedParams.DefaultMaterial);
 }
 
 const FText3DCachedMesh* FText3DFontFaceCache::FindGlyphMesh(uint32 InGlyphHash) const
 {
-	if (const FText3DCachedMesh* CachedMesh = GlyphMeshes.Find(InGlyphHash))
-	{
-		return CachedMesh;
-	}
-
-	return nullptr;
+	return GlyphMeshes.Find(InGlyphHash);
 }
 
 FText3DCachedMesh* FText3DFontFaceCache::FindGlyphMesh(uint32 InGlyphHash)
 {
-	if (FText3DCachedMesh* CachedMesh = GlyphMeshes.Find(InGlyphHash))
-	{
-		return CachedMesh;
-	}
-
-	return nullptr;
+	return GlyphMeshes.Find(InGlyphHash);
 }
 
 uint32 FText3DFontFaceCache::CleanupCache()
@@ -284,35 +309,34 @@ uint32 FText3DFontFaceCache::CleanupCache()
 	return GlyphMeshRemoved;
 }
 
-TText3DGlyphContourNodeShared FText3DFontFaceCache::LoadGlyphContours(uint32 InGlyphIndex) const
+TOptional<FText3DGlyphContourNode> FText3DFontFaceCache::LoadGlyphContours(uint32 InGlyphIndex) const
 {
 #if WITH_FREETYPE
 	if (!FontFace)
 	{
-		return nullptr;
+		return {};
 	}
 
 	const FT_Face Face = FontFace->GetFace();
 	if (!Face)
 	{
 		UE_LOG(LogText3D, Error, TEXT("Failed to load font face glyph contours '%u %i' due to invalid face"), FontFaceHash, InGlyphIndex);
-		return nullptr;
+		return {};
 	}
 
 	const FT_Error Result = FT_Load_Glyph(Face, InGlyphIndex, FT_LOAD_DEFAULT);
 	if (Result != 0)
 	{
 		UE_LOG(LogText3D, Error, TEXT("Failed to load font face glyph contours '%u %i' with error code : %i"), FontFaceHash, InGlyphIndex, Result);
-		return nullptr;
+		return {};
 	}
 
-	const FText3DGlyphLoader GlyphLoader(Face->glyph);
-	TText3DGlyphContourNodeShared Root = GlyphLoader.GetContourList();
-
-	return Root;
+	FText3DGlyphLoader GlyphLoader;
+	TOptional<FText3DGlyphContourNode> RootNode = GlyphLoader.GenerateContourList(Face->glyph);
+	return RootNode;
 #else
 	UE_LOG(LogText3D, Warning, TEXT("FreeType is not available, cannot proceed without it"))
-	return nullptr;
+	return {};
 #endif
 }
 
diff --git a/Engine/Plugins/VirtualProduction/Text3D/Source/Text3D/Private/Subsystems/Text3DEngineSubsystem.h b/Engine/Plugins/VirtualProduction/Text3D/Source/Text3D/Private/Subsystems/Text3DEngineSubsystem.h
--- a/Engine/Plugins/VirtualProduction/Text3D/Source/Text3D/Private/Subsystems/Text3DEngineSubsystem.h
+++ b/Engine/Plugins/VirtualProduction/Text3D/Source/Text3D/Private/Subsystems/Text3DEngineSubsystem.h
@@ -6,11 +6,11 @@
 #include "Text3D/Private/GeometryBuilders/Text3DGlyphContourNode.h"
 #include "Subsystems/EngineSubsystem.h"
 #include "Text3DTypes.h"
-
 #include "Text3DEngineSubsystem.generated.h"
 
-struct FSharedStruct;
 class FFreeTypeFace;
+class UText3DEngineSubsystem;
+struct FSharedStruct;
 
 USTRUCT()
 struct FText3DFontFaceCache
@@ -29,7 +29,31 @@ struct FText3DFontFaceCache
 	void UpdateFontFace(const TSharedPtr<FFreeTypeFace>& InFontFace);
 #endif
 
-	FText3DCachedMesh* FindOrBuildGlyphMesh(uint32 InGlyphIndex, const FGlyphMeshParameters& InParameters);
+	/** 
+	 * Loads the glyph contours and allocates the cached meshes the BuildGlyphMesh will be using (across multiple threads)
+	 * This happens in the game thread.
+	 */
+	TOptional<FText3DGlyphContourNode> PrepareCachedMesh(uint32 InGlyphIndex, const FGlyphMeshParameters& InParameters);
+
+	/** Data that does not change for a batch build */
+	struct FBuildGlyphMeshSharedDesc
+	{
+		TNotNull<UText3DEngineSubsystem*> TextSubsystem;
+		TNotNull<const FGlyphMeshParameters*> GlyphMeshParameters;
+		UMaterialInterface* DefaultMaterial;
+	};
+	/** Data that changes per build call in a batch */
+	struct FBuildGlyphMeshInstanceDesc
+	{
+		TNotNull<const FText3DGlyphContourNode*> RootNode;
+		uint32 GlyphIndex;
+	};
+	/** 
+	 * Builds a glyph mesh and saves the mesh (static and optionally dynamic) to the cached mesh struct 
+	 * This is called from multiple threads.
+	 */
+	void BuildGlyphMesh(const FBuildGlyphMeshSharedDesc& InSharedParams, const FBuildGlyphMeshInstanceDesc& InInstanceParams);
+
 	const FText3DCachedMesh* FindGlyphMesh(uint32 InGlyphHash) const;
 	FText3DCachedMesh* FindGlyphMesh(uint32 InGlyphHash);
 	uint32 CleanupCache();
@@ -55,7 +79,7 @@ struct FText3DFontFaceCache
 	}
 
 private:
-	TText3DGlyphContourNodeShared LoadGlyphContours(uint32 InGlyphIndex) const;
+	TOptional<FText3DGlyphContourNode> LoadGlyphContours(uint32 InGlyphIndex) const;
 
 #if WITH_FREETYPE
 	/** Font face for a Font+Typeface */
@@ -84,6 +108,7 @@ public:
 
 #if WITH_FREETYPE
 	FText3DFontFaceCache* FindOrAddCachedFontFace(const TSharedPtr<FFreeTypeFace>& InFontFace);
+	FText3DFontFaceCache* FindCachedFontFace(const TSharedPtr<FFreeTypeFace>& InFontFace);
 #endif
 
 	FText3DFontFaceCache* FindCachedFontFace(uint32 InFontFaceHash);
diff --git a/Engine/Plugins/VirtualProduction/Text3D/Source/Text3D/Private/Text3DTypes.cpp b/Engine/Plugins/VirtualProduction/Text3D/Source/Text3D/Private/Text3DTypes.cpp
--- a/Engine/Plugins/VirtualProduction/Text3D/Source/Text3D/Private/Text3DTypes.cpp
+++ b/Engine/Plugins/VirtualProduction/Text3D/Source/Text3D/Private/Text3DTypes.cpp
@@ -6,17 +6,14 @@
 
 namespace UE::Text3D::Geometry
 {
-	FCachedFontFaceGlyphHandle::FCachedFontFaceGlyphHandle(FText3DFontFaceCache* InFontFaceCache, uint32 InGlyphIndex, const FGlyphMeshParameters& InParameters)
+	FCachedFontFaceGlyphHandle::FCachedFontFaceGlyphHandle(const FText3DFontFaceCache& InFontFaceCache, uint32 InGlyphIndex, const FGlyphMeshParameters& InGlyphMeshParameters)
 	{
-		if (InFontFaceCache)
+		GlyphIndex = InGlyphIndex;
+		FontFaceHash = GetTypeHash(InFontFaceCache);
+		GlyphMeshHash = FText3DFontFaceCache::GetGlyphMeshHash(InGlyphIndex, InGlyphMeshParameters);
+		if (FText3DCachedMesh* CachedMesh = Resolve())
 		{
-			if (FText3DCachedMesh* CachedMesh = InFontFaceCache->FindOrBuildGlyphMesh(InGlyphIndex, InParameters))
-			{
-				FontFaceHash = GetTypeHash(*InFontFaceCache);
-				GlyphIndex = InGlyphIndex;
-				GlyphMeshHash = FText3DFontFaceCache::GetGlyphMeshHash(InGlyphIndex, InParameters);
-				CachedMesh->RefCount++;
-			}
+			CachedMesh->RefCount++;
 		}
 	}
 
diff --git a/Engine/Plugins/VirtualProduction/Text3D/Source/Text3D/Public/Extensions/Text3DDefaultGeometryExtension.h b/Engine/Plugins/VirtualProduction/Text3D/Source/Text3D/Public/Extensions/Text3DDefaultGeometryExtension.h
--- a/Engine/Plugins/VirtualProduction/Text3D/Source/Text3D/Public/Extensions/Text3DDefaultGeometryExtension.h
+++ b/Engine/Plugins/VirtualProduction/Text3D/Source/Text3D/Public/Extensions/Text3DDefaultGeometryExtension.h
@@ -3,7 +3,6 @@
 #pragma once
 
 #include "Text3DGeometryExtensionBase.h"
-#include "Text3D/Private/Subsystems/Text3DEngineSubsystem.h"
 #include "Text3DDefaultGeometryExtension.generated.h"
 
 UCLASS(MinimalAPI)
diff --git a/Engine/Plugins/VirtualProduction/Text3D/Source/Text3D/Public/Renderers/Text3DRendererBase.h b/Engine/Plugins/VirtualProduction/Text3D/Source/Text3D/Public/Renderers/Text3DRendererBase.h
--- a/Engine/Plugins/VirtualProduction/Text3D/Source/Text3D/Public/Renderers/Text3DRendererBase.h
+++ b/Engine/Plugins/VirtualProduction/Text3D/Source/Text3D/Public/Renderers/Text3DRendererBase.h
@@ -34,8 +34,11 @@ public:
 	/** Get cached bounds from last update */
 	FBox GetBounds() const;
 
+	/** Get the mesh type that this renderer uses */
+	virtual EText3DMeshType GetMeshType() const PURE_VIRTUAL(UText3DRendererBase::GetMeshType, return EText3DMeshType::Unknown; );
+
 	/** Get the implementation name for debug purposes */
-	virtual FName GetFriendlyName() const PURE_VIRTUAL(UText3DRenderingImplementationBase::GetFriendlyName, return NAME_None; );
+	virtual FName GetFriendlyName() const PURE_VIRTUAL(UText3DRendererBase::GetFriendlyName, return NAME_None; );
 
 	/** Is the renderer initialized (render state created) */
 	bool IsInitialized() const
diff --git a/Engine/Plugins/VirtualProduction/Text3D/Source/Text3D/Public/Text3DTypes.h b/Engine/Plugins/VirtualProduction/Text3D/Source/Text3D/Public/Text3DTypes.h
--- a/Engine/Plugins/VirtualProduction/Text3D/Source/Text3D/Public/Text3DTypes.h
+++ b/Engine/Plugins/VirtualProduction/Text3D/Source/Text3D/Public/Text3DTypes.h
@@ -77,6 +77,14 @@ enum class EText3DOutlineType : uint8
 	Fill
 };
 
+UENUM()
+enum class EText3DMeshType : uint8
+{
+	Unknown,
+	Static,
+	Dynamic,
+};
+
 UENUM()
 enum class EText3DGroupType : uint8
 {
@@ -265,6 +273,9 @@ struct FGlyphMeshParameters
 	UPROPERTY()
 	FVector PivotOffset = FVector::ZeroVector;
 
+	UPROPERTY()
+	EText3DMeshType MeshType = EText3DMeshType::Unknown;
+
 	friend uint32 GetTypeHash(const FGlyphMeshParameters& InElement)
 	{
 		uint32 HashParameters = 0;
@@ -276,6 +287,7 @@ struct FGlyphMeshParameters
 		HashParameters = HashCombine(HashParameters, GetTypeHash(InElement.OutlineExpand));
 		HashParameters = HashCombine(HashParameters, GetTypeHash(InElement.OutlineType));
 		HashParameters = HashCombine(HashParameters, GetTypeHash(InElement.PivotOffset));
+		HashParameters = HashCombine(HashParameters, GetTypeHash(InElement.MeshType));
 		return HashParameters;
 	}
 };
@@ -378,7 +390,7 @@ namespace UE::Text3D
 		struct FCachedFontFaceGlyphHandle
 		{
 			FCachedFontFaceGlyphHandle() = default;
-			FCachedFontFaceGlyphHandle(FText3DFontFaceCache* InFontFaceCache, uint32 InGlyphIndex, const FGlyphMeshParameters& InParameters);
+			FCachedFontFaceGlyphHandle(const FText3DFontFaceCache& InFontFaceCache, uint32 InGlyphIndex, const FGlyphMeshParameters& InGlyphMeshParameters);
 			~FCachedFontFaceGlyphHandle();
 
 			FCachedFontFaceGlyphHandle& operator=(const FCachedFontFaceGlyphHandle& InOther);
