Text3D: Moved generate contour list to be done across multiple threads. This was done by copying the outline data from FT_Load_Glyph (since this data gets removed for the next glyph, and a callt o FT_Load_Glyph from multiple threads for the same face is not supported). This copied data is then used to generate the contour list and build the meshes in each thread.

#Synced-CL 50867481

diff --git a/Engine/Plugins/VirtualProduction/Text3D/Source/Text3D/Private/Extensions/Text3DDefaultLayoutExtension.cpp b/Engine/Plugins/VirtualProduction/Text3D/Source/Text3D/Private/Extensions/Text3DDefaultLayoutExtension.cpp
--- a/Engine/Plugins/VirtualProduction/Text3D/Source/Text3D/Private/Extensions/Text3DDefaultLayoutExtension.cpp
+++ b/Engine/Plugins/VirtualProduction/Text3D/Source/Text3D/Private/Extensions/Text3DDefaultLayoutExtension.cpp
@@ -12,6 +12,7 @@
 #include "Extensions/Text3DStyleExtensionBase.h"
 #include "Framework/Text/RichTextLayoutMarshaller.h"
 #include "Framework/Text/RichTextMarkupProcessing.h"
+#include "GeometryBuilders/Text3DGlyphOutline.h"
 #include "LayoutBuilders/Text3DLayout.h"
 #include "LayoutBuilders/Text3DLayoutShaper.h"
 #include "LayoutBuilders/Text3DShapedGlyphText.h"
@@ -210,6 +211,8 @@ EText3DExtensionResult UText3DDefaultLayoutExtension::PreRendererUpdate(const UE
 {
 	using namespace UE::Text3D::Layout;
 
+	TRACE_CPUPROFILER_EVENT_SCOPE(UText3DDefaultLayoutExtension::PreRendererUpdate);
+
 	if (InParameters.CurrentFlag != EText3DRendererFlags::Geometry && InParameters.CurrentFlag != EText3DRendererFlags::Layout)
 	{
 		return EText3DExtensionResult::Active;
@@ -315,39 +318,42 @@ EText3DExtensionResult UText3DDefaultLayoutExtension::PreRendererUpdate(const UE
 	{
 		struct FBuildGlyphMeshDesc
 		{
-			FText3DGlyphContourNode RootNode;
+			UE::Text3D::FGlyphOutline GlyphOutline;
 			TSharedPtr<FFreeTypeFace> FontFace;
 			uint32 GlyphIndex = 0;
 		};
 		TArray<FBuildGlyphMeshDesc> BuildGlyphMeshDescs;
 
-		for (const TConstEnumerateRef<FGlyphLine> GlyphLine : EnumerateRange(GlyphText->Lines))
 		{
-			for (const TConstEnumerateRef<FGlyphEntry> GlyphEntry : EnumerateRange(GlyphLine->Glyphs))
+			TRACE_CPUPROFILER_EVENT_SCOPE(UText3DDefaultLayoutExtension::PrepareCachedMeshes);
+			for (const TConstEnumerateRef<FGlyphLine> GlyphLine : EnumerateRange(GlyphText->Lines))
 			{
-				if (!FGlyphText::IsValidGlyph(*GlyphEntry))
+				for (const TConstEnumerateRef<FGlyphEntry> GlyphEntry : EnumerateRange(GlyphLine->Glyphs))
 				{
-					continue;
-				}
+					if (!FGlyphText::IsValidGlyph(*GlyphEntry))
+					{
+						continue;
+					}
 
-				if (!Text3DComponent->GetCharacter(CharacterIndex))
-				{
-					UE_LOG(LogText3D, Warning, TEXT("Invalid character object returned for index %i after layout calculations, cannot proceed"), CharacterIndex)
-					return EText3DExtensionResult::Failed;
-				}
+					if (!Text3DComponent->GetCharacter(CharacterIndex))
+					{
+						UE_LOG(LogText3D, Warning, TEXT("Invalid character object returned for index %i after layout calculations, cannot proceed"), CharacterIndex)
+						return EText3DExtensionResult::Failed;
+					}
 
-				const TSharedPtr<FFreeTypeFace>& FontFace = GlyphText->FontFaces[CharacterIndex++];
-				if (FText3DFontFaceCache* FontFaceCache = TextSubsystem->FindOrAddCachedFontFace(FontFace))
-				{
-					if (TOptional<FText3DGlyphContourNode> RootNode = FontFaceCache->PrepareCachedMesh(GlyphEntry->Entry.GlyphIndex, GlyphMeshParameters))
+					const TSharedPtr<FFreeTypeFace>& FontFace = GlyphText->FontFaces[CharacterIndex++];
+					if (FText3DFontFaceCache* FontFaceCache = TextSubsystem->FindOrAddCachedFontFace(FontFace))
 					{
-						BuildGlyphMeshDescs.Emplace(MoveTemp(*RootNode), FontFace, GlyphEntry->Entry.GlyphIndex);
+						if (TOptional<UE::Text3D::FGlyphOutline> GlyphOutline = FontFaceCache->PrepareCachedMesh(GlyphEntry->Entry.GlyphIndex, GlyphMeshParameters))
+						{
+							BuildGlyphMeshDescs.Emplace(MoveTemp(*GlyphOutline), FontFace, GlyphEntry->Entry.GlyphIndex);
+						}
+					}
+					else
+					{
+						UE_LOG(LogText3D, Warning, TEXT("Failed to get cached font data for '%u %i' in Text3D geometry extension"), FText3DFontFaceCache::GetFontFaceHash(FontFace), GlyphEntry->Entry.GlyphIndex);
+						return EText3DExtensionResult::Failed;
 					}
-				}
-				else
-				{
-					UE_LOG(LogText3D, Warning, TEXT("Failed to get cached font data for '%u %i' in Text3D geometry extension"), FText3DFontFaceCache::GetFontFaceHash(FontFace), GlyphEntry->Entry.GlyphIndex);
-					return EText3DExtensionResult::Failed;
 				}
 			}
 		}
@@ -369,7 +375,7 @@ EText3DExtensionResult UText3DDefaultLayoutExtension::PreRendererUpdate(const UE
 				{
 					const FText3DFontFaceCache::FBuildGlyphMeshInstanceDesc InstanceParams
 						{
-							.RootNode = &BuildGlyphMeshDesc.RootNode,
+							.GlyphOutline = &BuildGlyphMeshDesc.GlyphOutline,
 							.GlyphIndex = BuildGlyphMeshDesc.GlyphIndex,
 						};
 					FontFaceCache->BuildGlyphMesh(SharedParams, InstanceParams);
diff --git a/Engine/Plugins/VirtualProduction/Text3D/Source/Text3D/Private/GeometryBuilders/Text3DGlyphLoader.cpp b/Engine/Plugins/VirtualProduction/Text3D/Source/Text3D/Private/GeometryBuilders/Text3DGlyphLoader.cpp
--- a/Engine/Plugins/VirtualProduction/Text3D/Source/Text3D/Private/GeometryBuilders/Text3DGlyphLoader.cpp
+++ b/Engine/Plugins/VirtualProduction/Text3D/Source/Text3D/Private/GeometryBuilders/Text3DGlyphLoader.cpp
@@ -7,6 +7,7 @@
 #include "GeometryBuilders/Text3DGlyphArrangement.h"
 #include "GeometryBuilders/Text3DGlyphPart.h"
 #include "Polygon2.h"
+#include "Text3DGlyphOutline.h"
 
 using namespace UE::Geometry;
 
@@ -17,16 +18,13 @@ FText3DGlyphLoader::FText3DGlyphLoader()
 }
 
 #if WITH_FREETYPE
-FText3DGlyphContourNode FText3DGlyphLoader::GenerateContourList(const FT_GlyphSlot InGlyph)
+FText3DGlyphContourNode FText3DGlyphLoader::GenerateContourList(const UE::Text3D::FGlyphOutline& InGlyphOutline)
 {
 	FText3DGlyphContourNode Root(/*Contour*/{}, /*CanHaveIntersections*/false, /*Clockwise*/true);
 
-	const FT_Outline Outline = InGlyph->outline;
-	const int32 ContourCount = Outline.n_contours;
-
-	for (int32 Index = 0; Index < ContourCount; Index++)
+	for (int16 ContourIndex : InGlyphOutline.Contours)
 	{
-		const TOptional<FPolygon2f> Contour = CreateContour(Outline, Index);
+		const TOptional<FPolygon2f> Contour = CreateContour(InGlyphOutline, ContourIndex);
 		if (Contour.IsSet())
 		{
 			Insert(FText3DGlyphContourNode(Contour, /*CanHaveIntersections*/true, Contour->IsClockwise()), Root);
@@ -287,9 +285,9 @@ FVector2D FText3DGlyphLoader::FCubicCurve::Tangent(const float T)
 }
 
 #if WITH_FREETYPE
-TOptional<FPolygon2f> FText3DGlyphLoader::CreateContour(const FT_Outline Outline, const int32 ContourIndex)
+TOptional<FPolygon2f> FText3DGlyphLoader::CreateContour(const UE::Text3D::FGlyphOutline& InGlyphOutline, int32 InContourIndex)
 {
-	TOptional<FPolygon2f> Contour = ProcessFreetypeOutline(Outline, ContourIndex);
+	TOptional<FPolygon2f> Contour = ProcessFreetypeOutline(InGlyphOutline, InContourIndex);
 	if (Contour.IsSet() && RemoveBadPoints(*Contour))
 	{
 		return Contour;
@@ -297,10 +295,10 @@ TOptional<FPolygon2f> FText3DGlyphLoader::CreateContour(const FT_Outline Outline
 	return {};
 }
 
-TOptional<FPolygon2f> FText3DGlyphLoader::ProcessFreetypeOutline(const FT_Outline Outline, const int32 ContourIndex)
+TOptional<FPolygon2f> FText3DGlyphLoader::ProcessFreetypeOutline(const UE::Text3D::FGlyphOutline& InGlyphOutline, int32 InContourIndex)
 {
 	const int32 StartIndex = EndIndex + 1;
-	EndIndex = Outline.contours[ContourIndex];
+	EndIndex = InContourIndex;
 	const int32 ContourLength = EndIndex - StartIndex + 1;
 
 	if (ContourLength < 3)
@@ -310,7 +308,7 @@ TOptional<FPolygon2f> FText3DGlyphLoader::ProcessFreetypeOutline(const FT_Outlin
 
 	ProcessedContour.Empty();
 
-	const FT_Vector* const Points = Outline.points + StartIndex;
+	const FT_Vector* const Points = InGlyphOutline.Points.GetData() + StartIndex;
 	auto ToFVector2D = [Points](const int32 Index)
 	{
 		const FT_Vector Point = Points[Index];
@@ -322,8 +320,7 @@ TOptional<FPolygon2f> FText3DGlyphLoader::ProcessFreetypeOutline(const FT_Outlin
 	FVector2D Next = ToFVector2D(0);
 	FVector2D NextNext = ToFVector2D(1);
 
-
-	const char* const Tags = Outline.tags + StartIndex;
+	const char* const Tags = InGlyphOutline.Tags.GetData() + StartIndex;
 	auto Tag = [Tags](int32 Index)
 	{
 		return FT_CURVE_TAG(Tags[Index]);
diff --git a/Engine/Plugins/VirtualProduction/Text3D/Source/Text3D/Private/GeometryBuilders/Text3DGlyphLoader.h b/Engine/Plugins/VirtualProduction/Text3D/Source/Text3D/Private/GeometryBuilders/Text3DGlyphLoader.h
--- a/Engine/Plugins/VirtualProduction/Text3D/Source/Text3D/Private/GeometryBuilders/Text3DGlyphLoader.h
+++ b/Engine/Plugins/VirtualProduction/Text3D/Source/Text3D/Private/GeometryBuilders/Text3DGlyphLoader.h
@@ -12,6 +12,11 @@
 struct FText3DGlyphArrangement;
 using UE::Geometry::FPolygon2f;
 
+namespace UE::Text3D
+{
+	struct FGlyphOutline;
+}
+
 /** Creates glyph geometry data from a font glyphs */
 class FText3DGlyphLoader final
 {
@@ -19,7 +24,7 @@ public:
 	explicit FText3DGlyphLoader();
 
 #if WITH_FREETYPE
-	FText3DGlyphContourNode GenerateContourList(const FT_GlyphSlot InGlyph);
+	FText3DGlyphContourNode GenerateContourList(const UE::Text3D::FGlyphOutline& InGlyphOutline);
 #endif
 
 private:
@@ -135,8 +140,8 @@ private:
 	TArray<int32> DividedContourIDs;
 
 #if WITH_FREETYPE
-	TOptional<FPolygon2f> CreateContour(const FT_Outline Outline, const int32 ContourIndex);
-	TOptional<FPolygon2f> ProcessFreetypeOutline(const FT_Outline Outline, const int32 ContourIndex);
+	TOptional<FPolygon2f> CreateContour(const UE::Text3D::FGlyphOutline& InGlyphOutline, int32 InContourIndex);
+	TOptional<FPolygon2f> ProcessFreetypeOutline(const UE::Text3D::FGlyphOutline& InGlyphOutline, int32 InContourIndex);
 #endif
 
 	/**
diff --git a/Engine/Plugins/VirtualProduction/Text3D/Source/Text3D/Private/GeometryBuilders/Text3DGlyphOutline.cpp b/Engine/Plugins/VirtualProduction/Text3D/Source/Text3D/Private/GeometryBuilders/Text3DGlyphOutline.cpp
new file mode 100644
--- /dev/null
+++ b/Engine/Plugins/VirtualProduction/Text3D/Source/Text3D/Private/GeometryBuilders/Text3DGlyphOutline.cpp
@@ -0,0 +1,20 @@
+// Copyright Epic Games, Inc. All Rights Reserved.
+
+#include "Text3DGlyphOutline.h"
+
+namespace UE::Text3D
+{
+
+#if WITH_FREETYPE
+FGlyphOutline FGlyphOutline::MakeOutline(const FT_Outline& InOutline)
+{
+	FGlyphOutline GlyphOutline;
+	GlyphOutline.Points   = TArray<FT_Vector>(InOutline.points, InOutline.n_points);
+	GlyphOutline.Tags     = TArray<char>(InOutline.tags, InOutline.n_points);
+	GlyphOutline.Contours = TArray<short>(InOutline.contours, InOutline.n_contours);
+	GlyphOutline.Flags    = InOutline.flags;
+	return GlyphOutline;
+}
+#endif
+
+} // UE::Text3D
diff --git a/Engine/Plugins/VirtualProduction/Text3D/Source/Text3D/Private/GeometryBuilders/Text3DGlyphOutline.h b/Engine/Plugins/VirtualProduction/Text3D/Source/Text3D/Private/GeometryBuilders/Text3DGlyphOutline.h
new file mode 100644
--- /dev/null
+++ b/Engine/Plugins/VirtualProduction/Text3D/Source/Text3D/Private/GeometryBuilders/Text3DGlyphOutline.h
@@ -0,0 +1,36 @@
+// Copyright Epic Games, Inc. All Rights Reserved.
+
+#pragma once
+
+#include "Containers/Array.h"
+#include "Misc/NotNull.h"
+#include "Text3DModule.h"
+
+namespace UE::Text3D
+{
+
+#if WITH_FREETYPE
+/**
+ * Mirrors the struct of 'FT_Outline' to be used outside of the loaded glyph.
+ * FT_Load_Glyph is not thread-safe when using the same FT_Face, so the data needs to be saved in this struct for 
+ * expensive computations to be done in multiple threads.
+ */
+struct FGlyphOutline
+{
+	static FGlyphOutline MakeOutline(const FT_Outline& InOutline);
+
+	/** Outline's points */
+	TArray<FT_Vector> Points;
+
+	/** Points flags */
+	TArray<char> Tags; 
+
+	/** Contour end points */
+	TArray<short> Contours;
+
+	/** Outline Masks */
+	int32 Flags = 0;
+};
+#endif
+
+} // UE::Text3D
diff --git a/Engine/Plugins/VirtualProduction/Text3D/Source/Text3D/Private/Subsystems/Text3DEngineSubsystem.cpp b/Engine/Plugins/VirtualProduction/Text3D/Source/Text3D/Private/Subsystems/Text3DEngineSubsystem.cpp
--- a/Engine/Plugins/VirtualProduction/Text3D/Source/Text3D/Private/Subsystems/Text3DEngineSubsystem.cpp
+++ b/Engine/Plugins/VirtualProduction/Text3D/Source/Text3D/Private/Subsystems/Text3DEngineSubsystem.cpp
@@ -3,16 +3,17 @@
 #include "Text3DEngineSubsystem.h"
 
 #include "Containers/Ticker.h"
-#include "UDynamicMesh.h"
 #include "Engine/Engine.h"
 #include "Engine/StaticMesh.h"
 #include "GeometryBuilders/Text3DGlyphContourNode.h"
 #include "GeometryBuilders/Text3DGlyphLoader.h"
 #include "GeometryBuilders/Text3DGlyphMeshBuilder.h"
+#include "GeometryBuilders/Text3DGlyphOutline.h"
 #include "Logs/Text3DLogs.h"
 #include "Materials/Material.h"
 #include "Misc/FileHelper.h"
 #include "Settings/Text3DProjectSettings.h"
+#include "UDynamicMesh.h"
 #include "UObject/ConstructorHelpers.h"
 
 #if WITH_FREETYPE
@@ -206,7 +207,7 @@ void FText3DFontFaceCache::UpdateFontFace(const TSharedPtr<FFreeTypeFace>& InFon
 }
 #endif
 
-TOptional<FText3DGlyphContourNode> FText3DFontFaceCache::PrepareCachedMesh(uint32 InGlyphIndex, const FGlyphMeshParameters& InParameters)
+TOptional<UE::Text3D::FGlyphOutline> FText3DFontFaceCache::PrepareCachedMesh(uint32 InGlyphIndex, const FGlyphMeshParameters& InParameters)
 {
 	const uint32 HashParameters = GetGlyphMeshHash(InGlyphIndex, InParameters);
 	if (FText3DCachedMesh* CachedMesh = GlyphMeshes.Find(HashParameters))
@@ -214,8 +215,8 @@ TOptional<FText3DGlyphContourNode> FText3DFontFaceCache::PrepareCachedMesh(uint3
 		return {};
 	}
 
-	TOptional<FText3DGlyphContourNode> GlyphContourRootNode = LoadGlyphContours(InGlyphIndex);
-	if (!GlyphContourRootNode.IsSet() || GlyphContourRootNode->Children.IsEmpty())
+	TOptional<UE::Text3D::FGlyphOutline> GlyphOutline = LoadGlyphOutline(InGlyphIndex);
+	if (!GlyphOutline.IsSet())
 	{
 		return {};
 	}
@@ -230,7 +231,7 @@ TOptional<FText3DGlyphContourNode> FText3DFontFaceCache::PrepareCachedMesh(uint3
 	}
 #endif
 
-	return GlyphContourRootNode;
+	return GlyphOutline;
 }
 
 uint32 FText3DFontFaceCache::GetGlyphMeshHash(uint32 InGlyphIndex, const FGlyphMeshParameters& InParameters)
@@ -240,6 +241,17 @@ uint32 FText3DFontFaceCache::GetGlyphMeshHash(uint32 InGlyphIndex, const FGlyphM
 
 void FText3DFontFaceCache::BuildGlyphMesh(const FBuildGlyphMeshSharedDesc& InSharedParams, const FBuildGlyphMeshInstanceDesc& InInstanceParams)
 {
+	TOptional<FText3DGlyphContourNode> RootNode;
+	{
+		FText3DGlyphLoader GlyphLoader;
+		RootNode = GlyphLoader.GenerateContourList(*InInstanceParams.GlyphOutline);
+	}
+
+	if (!RootNode.IsSet() || RootNode->Children.IsEmpty())
+	{
+		return;
+	}
+
 	const FGlyphMeshParameters& GlyphMeshParameters = *InSharedParams.GlyphMeshParameters;
 
 	const uint32 HashParameters = GetGlyphMeshHash(InInstanceParams.GlyphIndex, GlyphMeshParameters);
@@ -270,7 +282,7 @@ void FText3DFontFaceCache::BuildGlyphMesh(const FBuildGlyphMeshSharedDesc& InSha
 
 	FText3DGlyphMeshBuilder MeshCreator;
 	MeshCreator.CreateMeshes(
-		*InInstanceParams.RootNode,
+		*RootNode,
 		GlyphMeshParameters.Extrude,
 		GlyphMeshParameters.Bevel,
 		GlyphMeshParameters.BevelType,
@@ -309,7 +321,7 @@ uint32 FText3DFontFaceCache::CleanupCache()
 	return GlyphMeshRemoved;
 }
 
-TOptional<FText3DGlyphContourNode> FText3DFontFaceCache::LoadGlyphContours(uint32 InGlyphIndex) const
+TOptional<UE::Text3D::FGlyphOutline> FText3DFontFaceCache::LoadGlyphOutline(uint32 InGlyphIndex) const
 {
 #if WITH_FREETYPE
 	if (!FontFace)
@@ -325,15 +337,13 @@ TOptional<FText3DGlyphContourNode> FText3DFontFaceCache::LoadGlyphContours(uint3
 	}
 
 	const FT_Error Result = FT_Load_Glyph(Face, InGlyphIndex, FT_LOAD_DEFAULT);
-	if (Result != 0)
+	if (Result != 0 || !Face->glyph)
 	{
 		UE_LOG(LogText3D, Error, TEXT("Failed to load font face glyph contours '%u %i' with error code : %i"), FontFaceHash, InGlyphIndex, Result);
 		return {};
 	}
 
-	FText3DGlyphLoader GlyphLoader;
-	TOptional<FText3DGlyphContourNode> RootNode = GlyphLoader.GenerateContourList(Face->glyph);
-	return RootNode;
+	return UE::Text3D::FGlyphOutline::MakeOutline(Face->glyph->outline);
 #else
 	UE_LOG(LogText3D, Warning, TEXT("FreeType is not available, cannot proceed without it"))
 	return {};
diff --git a/Engine/Plugins/VirtualProduction/Text3D/Source/Text3D/Private/Subsystems/Text3DEngineSubsystem.h b/Engine/Plugins/VirtualProduction/Text3D/Source/Text3D/Private/Subsystems/Text3DEngineSubsystem.h
--- a/Engine/Plugins/VirtualProduction/Text3D/Source/Text3D/Private/Subsystems/Text3DEngineSubsystem.h
+++ b/Engine/Plugins/VirtualProduction/Text3D/Source/Text3D/Private/Subsystems/Text3DEngineSubsystem.h
@@ -8,6 +8,11 @@
 #include "Text3DTypes.h"
 #include "Text3DEngineSubsystem.generated.h"
 
+namespace UE::Text3D
+{
+	struct FGlyphOutline;
+}
+
 class FFreeTypeFace;
 class UText3DEngineSubsystem;
 struct FSharedStruct;
@@ -33,7 +38,7 @@ struct FText3DFontFaceCache
 	 * Loads the glyph contours and allocates the cached meshes the BuildGlyphMesh will be using (across multiple threads)
 	 * This happens in the game thread.
 	 */
-	TOptional<FText3DGlyphContourNode> PrepareCachedMesh(uint32 InGlyphIndex, const FGlyphMeshParameters& InParameters);
+	TOptional<UE::Text3D::FGlyphOutline> PrepareCachedMesh(uint32 InGlyphIndex, const FGlyphMeshParameters& InParameters);
 
 	/** Data that does not change for a batch build */
 	struct FBuildGlyphMeshSharedDesc
@@ -45,7 +50,7 @@ struct FText3DFontFaceCache
 	/** Data that changes per build call in a batch */
 	struct FBuildGlyphMeshInstanceDesc
 	{
-		TNotNull<const FText3DGlyphContourNode*> RootNode;
+		TNotNull<const UE::Text3D::FGlyphOutline*> GlyphOutline;
 		uint32 GlyphIndex;
 	};
 	/** 
@@ -79,7 +84,7 @@ struct FText3DFontFaceCache
 	}
 
 private:
-	TOptional<FText3DGlyphContourNode> LoadGlyphContours(uint32 InGlyphIndex) const;
+	TOptional<UE::Text3D::FGlyphOutline> LoadGlyphOutline(uint32 InGlyphIndex) const;
 
 #if WITH_FREETYPE
 	/** Font face for a Font+Typeface */
