Text3D: Changed mode so that BatchAddPrimitives is not called by the static mesh renderer and allow the add primitve to be called in end of frame updates

#Synced-CL 50867481

diff --git a/Engine/Plugins/VirtualProduction/Text3D/Source/Text3D/Private/Renderers/StaticMeshes/Text3DStaticMeshesRenderer.cpp b/Engine/Plugins/VirtualProduction/Text3D/Source/Text3D/Private/Renderers/StaticMeshes/Text3DStaticMeshesRenderer.cpp
--- a/Engine/Plugins/VirtualProduction/Text3D/Source/Text3D/Private/Renderers/StaticMeshes/Text3DStaticMeshesRenderer.cpp
+++ b/Engine/Plugins/VirtualProduction/Text3D/Source/Text3D/Private/Renderers/StaticMeshes/Text3DStaticMeshesRenderer.cpp
@@ -24,12 +24,14 @@
 
 namespace UE::Text3D::Private
 {
-	// Static mesh components will be added in batch to the FScene, rather than one by one
-	static bool bBatchRegisterStaticMeshes = true;
-	static FAutoConsoleVariableRef CVarBatchRegisterStaticMeshes(
-		TEXT("Text3D.BatchRegisterStaticMeshes"),
-		bBatchRegisterStaticMeshes,
-		TEXT("Determines whether static mesh component primitives should be batch added to the scene, rather than one by one. True to enable it, false to revert back to legacy behavior"),
+	static int32 RegisterStaticMeshesMode = 2;
+	static FAutoConsoleVariableRef CVarRegisterStaticMeshes(
+		TEXT("Text3D.RegisterStaticMeshesMode"),
+		RegisterStaticMeshesMode,
+		TEXT("Determines how static mesh component primitives should be batch added to the scene.\n")
+		TEXT("0: Disabled. Primitives will be added one by one, not in batch.\n")
+		TEXT("1: Immediate. Primitives will be added in batch to the scene as soon as all the static meshes are set.\n")
+		TEXT("2: Deferred. Primitives will be added to the scene on end of frame updates.\n"),
 		ECVF_Default);
 }
 
@@ -80,12 +82,16 @@ void UText3DStaticMeshesRenderer::OnUpdate(const UE::Text3D::Renderer::FUpdatePa
 
 		const TArrayView<UPrimitiveComponent*> StaticMeshComponents = MakeArrayView(reinterpret_cast<UPrimitiveComponent**>(CharacterMeshes.GetData()), CharacterMeshes.Num());
 
-		const bool bBatchRegisterStaticMeshes = UE::Text3D::Private::bBatchRegisterStaticMeshes
+		const bool bShouldBulkReregister = World 
+			&& !StaticMeshComponents.IsEmpty() 
+			&& UE::Text3D::Private::RegisterStaticMeshesMode > 0
 			&& !StaticMeshComponents.ContainsByPredicate([](UPrimitiveComponent* InComponent)
 				{
 					return !InComponent || !InComponent->IsRegistered();
 				});
 
+		const bool bShouldBatchAddMeshes = bShouldBulkReregister && UE::Text3D::Private::RegisterStaticMeshesMode == 1;
+
 		auto SetBulkReregister = [&StaticMeshComponents](bool bInNewValue)
 			{
 				for (UPrimitiveComponent* Component : StaticMeshComponents)
@@ -97,10 +103,13 @@ void UText3DStaticMeshesRenderer::OnUpdate(const UE::Text3D::Renderer::FUpdatePa
 				}
 			};
 
-		if (bBatchRegisterStaticMeshes && World && !StaticMeshComponents.IsEmpty())
+		if (bShouldBulkReregister)
 		{
-			World->Scene->BatchRemovePrimitives(StaticMeshComponents);
-			World->Scene->BatchReleasePrimitives(StaticMeshComponents);
+			if (bShouldBatchAddMeshes)
+			{
+				World->Scene->BatchRemovePrimitives(StaticMeshComponents);
+				World->Scene->BatchReleasePrimitives(StaticMeshComponents);
+			}
 			SetBulkReregister(true);
 		}
 
@@ -135,12 +144,14 @@ void UText3DStaticMeshesRenderer::OnUpdate(const UE::Text3D::Renderer::FUpdatePa
 			}
 		});
 
-		if (bBatchRegisterStaticMeshes && World && !StaticMeshComponents.IsEmpty())
+		if (bShouldBulkReregister)
 		{
-			World->Scene->BatchAddPrimitives(StaticMeshComponents);
+			if (bShouldBatchAddMeshes)
+			{
+				World->Scene->BatchAddPrimitives(StaticMeshComponents);
+			}
 			SetBulkReregister(false);
 		}
-
 		RefreshBounds();
 	}
 	else if (InParameters.CurrentFlag == EText3DRendererFlags::Layout)
