Text3D: save default material in text3d subsystem to avoid synchronous material load in text3d building

#Synced-CL 50950138

diff --git a/Engine/Plugins/VirtualProduction/Text3D/Source/Text3D/Private/Extensions/Text3DDefaultLayoutExtension.cpp b/Engine/Plugins/VirtualProduction/Text3D/Source/Text3D/Private/Extensions/Text3DDefaultLayoutExtension.cpp
--- a/Engine/Plugins/VirtualProduction/Text3D/Source/Text3D/Private/Extensions/Text3DDefaultLayoutExtension.cpp
+++ b/Engine/Plugins/VirtualProduction/Text3D/Source/Text3D/Private/Extensions/Text3DDefaultLayoutExtension.cpp
@@ -363,7 +363,7 @@ EText3DExtensionResult UText3DDefaultLayoutExtension::PreRendererUpdate(const UE
 			{
 				.TextSubsystem = TextSubsystem,
 				.GlyphMeshParameters = &GlyphMeshParameters,
-				.DefaultMaterial = GetDefault<UText3DProjectSettings>()->GetDefaultMaterial(),
+				.DefaultMaterial = TextSubsystem->GetDefaultMaterial(),
 			};
 
 		ParallelFor(BuildGlyphMeshDescs.Num(), 
diff --git a/Engine/Plugins/VirtualProduction/Text3D/Source/Text3D/Private/Settings/Text3DProjectSettings.cpp b/Engine/Plugins/VirtualProduction/Text3D/Source/Text3D/Private/Settings/Text3DProjectSettings.cpp
--- a/Engine/Plugins/VirtualProduction/Text3D/Source/Text3D/Private/Settings/Text3DProjectSettings.cpp
+++ b/Engine/Plugins/VirtualProduction/Text3D/Source/Text3D/Private/Settings/Text3DProjectSettings.cpp
@@ -266,6 +266,11 @@ FName UText3DProjectSettings::GetDebugModePropertyName()
 	return GET_MEMBER_NAME_CHECKED(UText3DProjectSettings, bDebugMode);
 }
 
+FName UText3DProjectSettings::GetDefaultMaterialPropertyName()
+{
+	return GET_MEMBER_NAME_CHECKED(UText3DProjectSettings, DefaultMaterial);
+}
+
 bool UText3DProjectSettings::GetDebugMode() const
 {
 	return bDebugMode;
diff --git a/Engine/Plugins/VirtualProduction/Text3D/Source/Text3D/Private/Subsystems/Text3DEngineSubsystem.cpp b/Engine/Plugins/VirtualProduction/Text3D/Source/Text3D/Private/Subsystems/Text3DEngineSubsystem.cpp
--- a/Engine/Plugins/VirtualProduction/Text3D/Source/Text3D/Private/Subsystems/Text3DEngineSubsystem.cpp
+++ b/Engine/Plugins/VirtualProduction/Text3D/Source/Text3D/Private/Subsystems/Text3DEngineSubsystem.cpp
@@ -34,12 +34,27 @@ void UText3DEngineSubsystem::Initialize(FSubsystemCollectionBase& InCollection)
 {
 	Super::Initialize(InCollection);
 	ScheduleCacheCleanup();
+
+	if (UText3DProjectSettings* Text3DSettings = UText3DProjectSettings::GetMutable())
+	{
+#if WITH_EDITOR
+		Text3DSettings->OnSettingChanged().AddUObject(this, &UText3DEngineSubsystem::OnText3DSettingsChanged);
+#endif
+		UpdateDefaultMaterial();
+	}
 }
 
 void UText3DEngineSubsystem::Deinitialize()
 {
 	Super::Deinitialize();
 	ClearCache();
+
+#if WITH_EDITOR
+	if (UText3DProjectSettings* Text3DSettings = UText3DProjectSettings::GetMutable())
+	{
+		Text3DSettings->OnSettingChanged().RemoveAll(this);
+	}
+#endif
 }
 
 bool UText3DEngineSubsystem::Exec_Dev(UWorld* InWorld, const TCHAR* InCmd, FOutputDevice& InAr)
@@ -70,6 +85,24 @@ bool UText3DEngineSubsystem::Exec_Dev(UWorld* InWorld, const TCHAR* InCmd, FOutp
 	return false;
 }
 
+#if WITH_EDITOR
+void UText3DEngineSubsystem::OnText3DSettingsChanged(UObject* InObject, FPropertyChangedEvent& InPropertyChangedEvent)
+{
+	if (InPropertyChangedEvent.GetMemberPropertyName() == UText3DProjectSettings::GetDefaultMaterialPropertyName())
+	{
+		UpdateDefaultMaterial();
+	}
+}
+#endif
+
+void UText3DEngineSubsystem::UpdateDefaultMaterial()
+{
+	if (const UText3DProjectSettings* Text3DSettings = UText3DProjectSettings::Get())
+	{
+		CachedDefaultMaterial = Text3DSettings->GetDefaultMaterial();
+	}
+}
+
 void UText3DEngineSubsystem::ScheduleCacheCleanup()
 {
 	CancelCacheCleanup();
@@ -95,6 +128,11 @@ void UText3DEngineSubsystem::CancelCacheCleanup()
 	}
 }
 
+UMaterial* UText3DEngineSubsystem::GetDefaultMaterial() const
+{
+	return CachedDefaultMaterial;
+}
+
 bool UText3DEngineSubsystem::CleanupCache(float)
 {
 	uint32 GlyphMeshCount = 0;
diff --git a/Engine/Plugins/VirtualProduction/Text3D/Source/Text3D/Private/Subsystems/Text3DEngineSubsystem.h b/Engine/Plugins/VirtualProduction/Text3D/Source/Text3D/Private/Subsystems/Text3DEngineSubsystem.h
--- a/Engine/Plugins/VirtualProduction/Text3D/Source/Text3D/Private/Subsystems/Text3DEngineSubsystem.h
+++ b/Engine/Plugins/VirtualProduction/Text3D/Source/Text3D/Private/Subsystems/Text3DEngineSubsystem.h
@@ -14,6 +14,7 @@ namespace UE::Text3D
 }
 
 class FFreeTypeFace;
+class UMaterial;
 class UText3DEngineSubsystem;
 struct FSharedStruct;
 
@@ -111,6 +112,9 @@ public:
 	void ScheduleCacheCleanup();
 	void CancelCacheCleanup();
 
+	/** Returns the default material for text meshes */
+	UMaterial* GetDefaultMaterial() const;
+
 #if WITH_FREETYPE
 	FText3DFontFaceCache* FindOrAddCachedFontFace(const TSharedPtr<FFreeTypeFace>& InFontFace);
 	FText3DFontFaceCache* FindCachedFontFace(const TSharedPtr<FFreeTypeFace>& InFontFace);
@@ -128,6 +132,13 @@ private:
 	virtual bool Exec_Dev(UWorld* InWorld, const TCHAR* InCmd, FOutputDevice& InAr) override;
 	// ~End FSelfRegisteringExec
 
+#if WITH_EDITOR
+	/** Called when the UText3DProjectSettings has changed */
+	void OnText3DSettingsChanged(UObject* InObject, FPropertyChangedEvent& InPropertyChangedEvent);
+#endif
+	/** Updates the cached default material */
+	void UpdateDefaultMaterial();
+
 	bool CleanupCache(float);
 	void PrintCache() const;
 	void ClearCache();
@@ -135,5 +146,9 @@ private:
 	UPROPERTY()
 	TMap<uint32, FText3DFontFaceCache> CachedFontFaces;
 
+	/** the default material to use for text meshes */
+	UPROPERTY()
+	TObjectPtr<UMaterial> CachedDefaultMaterial;
+
 	FTSTicker::FDelegateHandle CacheCleanupDelegate;
 };
diff --git a/Engine/Plugins/VirtualProduction/Text3D/Source/Text3D/Public/Settings/Text3DProjectSettings.h b/Engine/Plugins/VirtualProduction/Text3D/Source/Text3D/Public/Settings/Text3DProjectSettings.h
--- a/Engine/Plugins/VirtualProduction/Text3D/Source/Text3D/Public/Settings/Text3DProjectSettings.h
+++ b/Engine/Plugins/VirtualProduction/Text3D/Source/Text3D/Public/Settings/Text3DProjectSettings.h
@@ -27,6 +27,7 @@ public:
 
 #if WITH_EDITOR
 	static FName GetDebugModePropertyName();
+	static FName GetDefaultMaterialPropertyName();
 #endif
 
 	struct FMaterialParameters
