Text3D: changed way cleanup is done. Before the cleanup would happen periodically (once every 5 minutes by default). Now, it will schedule a cleanup to happen after a set amount of seconds when a glyph becomes unused.
Additionally, fixed an issue where async-built meshes had subobjects that would still contain the 'Async' flag. This prevented these meshes from being properly gc'd.

#Synced-CL 50950138

diff --git a/Engine/Plugins/VirtualProduction/Text3D/Source/Text3D/Private/Extensions/Text3DDefaultLayoutExtension.cpp b/Engine/Plugins/VirtualProduction/Text3D/Source/Text3D/Private/Extensions/Text3DDefaultLayoutExtension.cpp
--- a/Engine/Plugins/VirtualProduction/Text3D/Source/Text3D/Private/Extensions/Text3DDefaultLayoutExtension.cpp
+++ b/Engine/Plugins/VirtualProduction/Text3D/Source/Text3D/Private/Extensions/Text3DDefaultLayoutExtension.cpp
@@ -369,6 +369,7 @@ EText3DExtensionResult UText3DDefaultLayoutExtension::PreRendererUpdate(const UE
 		ParallelFor(BuildGlyphMeshDescs.Num(), 
 			[&SharedParams, &BuildGlyphMeshDescs](int32 InTaskInfoIndex)
 			{
+				TRACE_CPUPROFILER_EVENT_SCOPE(UText3DDefaultLayoutExtension::BuildGlyphMesh);
 				const FBuildGlyphMeshDesc& BuildGlyphMeshDesc = BuildGlyphMeshDescs[InTaskInfoIndex];
 
 				if (FText3DFontFaceCache* FontFaceCache = SharedParams.TextSubsystem->FindCachedFontFace(BuildGlyphMeshDesc.FontFace))
diff --git a/Engine/Plugins/VirtualProduction/Text3D/Source/Text3D/Private/Settings/Text3DProjectSettings.cpp b/Engine/Plugins/VirtualProduction/Text3D/Source/Text3D/Private/Settings/Text3DProjectSettings.cpp
--- a/Engine/Plugins/VirtualProduction/Text3D/Source/Text3D/Private/Settings/Text3DProjectSettings.cpp
+++ b/Engine/Plugins/VirtualProduction/Text3D/Source/Text3D/Private/Settings/Text3DProjectSettings.cpp
@@ -84,34 +84,25 @@ UMaterial* UText3DProjectSettings::GetBaseMaterial(const FText3DMaterialKey& InK
 
 void UText3DProjectSettings::SetScheduleFontFaceGlyphCleanup(bool bInEnabled)
 {
-	if (bScheduleFontFaceGlyphCleanup == bInEnabled)
-	{
-		return;
-	}
-
-	bScheduleFontFaceGlyphCleanup = bInEnabled;
-	OnFontFaceCleanupOptionsChanged();
 }
 
 bool UText3DProjectSettings::GetScheduleFontFaceGlyphCleanup() const
 {
-	return bScheduleFontFaceGlyphCleanup;
+	return false;
 }
 
 void UText3DProjectSettings::SetFontFaceGlyphCleanupPeriod(float InValue)
 {
-	if (FMath::IsNearlyEqual(FontFaceGlyphCleanupPeriod, InValue))
-	{
-		return;
-	}
-
-	FontFaceGlyphCleanupPeriod = InValue;
-	OnFontFaceCleanupOptionsChanged();
 }
 
 float UText3DProjectSettings::GetFontFaceGlyphCleanupPeriod() const
 {
-	return FontFaceGlyphCleanupPeriod;
+	return 0.f;
+}
+
+void UText3DProjectSettings::SetFontFaceGlyphCleanupDelay(float InDelay)
+{
+	FontFaceGlyphCleanupDelay = InDelay;
 }
 
 UText3DProjectSettings::UText3DProjectSettings()
@@ -153,14 +144,6 @@ void UText3DProjectSettings::AddMaterial(FText3DMaterialKey&& InKey, TSoftObject
 	BaseMaterials.Emplace(InKey, InMaterial);
 }
 
-void UText3DProjectSettings::OnFontFaceCleanupOptionsChanged()
-{
-	if (UText3DEngineSubsystem* TextSubsystem = UText3DEngineSubsystem::Get())
-	{
-		TextSubsystem->ScheduleCacheCleanup();
-	}
-}
-
 #if WITH_EDITOR
 void UText3DProjectSettings::PostEditChangeProperty(struct FPropertyChangedEvent& InChangeEvent)
 {
@@ -177,25 +160,15 @@ void UText3DProjectSettings::PostEditChangeProperty(struct FPropertyChangedEvent
 		GET_MEMBER_NAME_CHECKED(UText3DProjectSettings, bShowOnlyBold),
 		GET_MEMBER_NAME_CHECKED(UText3DProjectSettings, bShowOnlyItalic),
 		GET_MEMBER_NAME_CHECKED(UText3DProjectSettings, bDebugMode),
-	};
-
-	static const TSet<FName> FontFaceCleanupPropertyNames
-	{
-		GET_MEMBER_NAME_CHECKED(UText3DProjectSettings, bScheduleFontFaceGlyphCleanup),
-		GET_MEMBER_NAME_CHECKED(UText3DProjectSettings, FontFaceGlyphCleanupPeriod)
+		GET_MEMBER_NAME_CHECKED(UText3DProjectSettings, FontFaceGlyphCleanupDelay),
 	};
 
 	const FName MemberName = InChangeEvent.GetMemberPropertyName();
 
-	if (PropertyNames.Contains(MemberName) || FontFaceCleanupPropertyNames.Contains(MemberName))
+	if (PropertyNames.Contains(MemberName))
 	{
 		Save();
 	}
-
-	if (FontFaceCleanupPropertyNames.Contains(MemberName))
-	{
-		OnFontFaceCleanupOptionsChanged();
-	}
 }
 
 TConstArrayView<FString> UText3DProjectSettings::GetFavoriteFonts() const
diff --git a/Engine/Plugins/VirtualProduction/Text3D/Source/Text3D/Private/Subsystems/Text3DEngineSubsystem.cpp b/Engine/Plugins/VirtualProduction/Text3D/Source/Text3D/Private/Subsystems/Text3DEngineSubsystem.cpp
--- a/Engine/Plugins/VirtualProduction/Text3D/Source/Text3D/Private/Subsystems/Text3DEngineSubsystem.cpp
+++ b/Engine/Plugins/VirtualProduction/Text3D/Source/Text3D/Private/Subsystems/Text3DEngineSubsystem.cpp
@@ -33,7 +33,6 @@ UText3DEngineSubsystem* UText3DEngineSubsystem::Get()
 void UText3DEngineSubsystem::Initialize(FSubsystemCollectionBase& InCollection)
 {
 	Super::Initialize(InCollection);
-	ScheduleCacheCleanup();
 
 	if (UText3DProjectSettings* Text3DSettings = UText3DProjectSettings::GetMutable())
 	{
@@ -105,26 +104,24 @@ void UText3DEngineSubsystem::UpdateDefaultMaterial()
 
 void UText3DEngineSubsystem::ScheduleCacheCleanup()
 {
-	CancelCacheCleanup();
-
-	if (const UText3DProjectSettings* TextSettings = UText3DProjectSettings::Get())
+	// Already queued up. Nothing else to do.
+	if (CacheCleanupDelegate.IsValid())
 	{
-		if (TextSettings->GetScheduleFontFaceGlyphCleanup())
-		{
-			CacheCleanupDelegate = FTSTicker::GetCoreTicker().AddTicker(
-				FTickerDelegate::CreateUObject(this, &UText3DEngineSubsystem::CleanupCache),
-				FMath::Max(0, TextSettings->GetFontFaceGlyphCleanupPeriod())
-			);
-		}
+		return;
 	}
-}
 
-void UText3DEngineSubsystem::CancelCacheCleanup()
-{
-	if (CacheCleanupDelegate.IsValid())
+	const UText3DProjectSettings* const TextSettings = UText3DProjectSettings::Get();
+	const float CleanupDelay = TextSettings ? TextSettings->GetFontFaceGlyphCleanupDelay() : 5.f;
+
+	// Immediate cleanup if close to zero or negative.
+	if (CleanupDelay <= UE_SMALL_NUMBER)
+	{
+		CleanupCache(0.f);
+	}
+	else
 	{
-		FTSTicker::RemoveTicker(CacheCleanupDelegate);
-		CacheCleanupDelegate.Reset();
+		CacheCleanupDelegate = FTSTicker::GetCoreTicker()
+			.AddTicker(FTickerDelegate::CreateUObject(this, &UText3DEngineSubsystem::CleanupCache), CleanupDelay);
 	}
 }
 
@@ -135,6 +132,11 @@ UMaterial* UText3DEngineSubsystem::GetDefaultMaterial() const
 
 bool UText3DEngineSubsystem::CleanupCache(float)
 {
+	// Removing CleanupHandle Ticker while still executing this ticker delegate is by-design valid.
+	// see TickerTests.cpp:112 "a delegate removal from inside the delegate execution (used to be a deadlock)"
+	FTSTicker::RemoveTicker(CacheCleanupDelegate);
+	CacheCleanupDelegate.Reset();
+
 	uint32 GlyphMeshCount = 0;
 	uint32 GlyphMeshRemoved = 0;
 	uint32 FontFaceRemoved = 0;
@@ -156,7 +158,8 @@ bool UText3DEngineSubsystem::CleanupCache(float)
 		UE_LOG(LogText3D, Verbose, TEXT("Text3D Engine Subsystem cache cleanup completed : %i unused glyph(s) removed, %i remaining glyph(s), %i unused font face(s) removed, %i remaining font face(s)"), GlyphMeshRemoved, GlyphMeshCount, FontFaceRemoved, CachedFontFaces.Num());
 	}
 
-	return true;
+	// one time off.
+	return false;
 }
 
 void UText3DEngineSubsystem::PrintCache() const
@@ -295,27 +298,46 @@ void FText3DFontFaceCache::BuildGlyphMesh(const FBuildGlyphMeshSharedDesc& InSha
 	const uint32 HashParameters = GetGlyphMeshHash(InInstanceParams.GlyphIndex, GlyphMeshParameters);
 	FText3DCachedMesh& CachedMesh = GlyphMeshes.FindChecked(HashParameters);
 
-	{
-		auto OnMeshAllocated = 
-			[bIsInGameThread = IsInGameThread()](TNotNull<UObject*> InMesh)
+	auto ClearAsyncFlags = [bIsInGameThread = IsInGameThread()](UObject* InObject)
+		{
+			if (!bIsInGameThread && InObject)
 			{
-				if (!bIsInGameThread)
-				{
-					// OK to clear the async flags for the objects immediately because they have already been assigned directly to a UPROPERTY.
-					InMesh->AtomicallyClearInternalFlags(EInternalObjectFlags::Async);
-				}
-			};
-
-		FGCScopeGuard ScopeGuard;
+				constexpr bool bIncludeNestedObjects = true;
+
+				// OK to clear the async flags for the objects immediately because they have already been assigned directly to a UPROPERTY.
+				InObject->AtomicallyClearInternalFlags(EInternalObjectFlags::Async);
+
+				// Subobjects created under the object might still be marked as async.
+				// E.g. for UStaticMesh these objects get created with async flag:
+				//   - BodySetup
+				//   - AssetImportData
+				//   - For HiResSourceModel:
+				//      - FStaticMeshSourceModel::StaticMeshDescriptionBulkData
+				//      - FStaticMeshSourceModel::StaticMeshDescriptionBulkData.PreallocatedMeshDescription
+				//   - For SourceModels:
+				//      - Only FStaticMeshSourceModel::StaticMeshDescriptionBulkData.PreallocatedMeshDescription
+				//        because the StaticMeshDescriptionBulkData has its async flag cleared in FStaticMeshSourceModel::CreateSubObjects.
+				ForEachObjectWithOuter(InObject, [](UObject* InSubobject)
+					{
+						InSubobject->AtomicallyClearInternalFlags(EInternalObjectFlags::Async);
+					}
+					, bIncludeNestedObjects);
+			}
+		};
 
-		CachedMesh.StaticMesh = NewObject<UStaticMesh>(InSharedParams.TextSubsystem, NAME_None, RF_Transient | RF_Public);
-		OnMeshAllocated(CachedMesh.StaticMesh);
+	// UObjects like BodySetup get created while building mesh descriptions, so scope will need to cover those too.
+	FGCScopeGuard ScopeGuard;
+	ON_SCOPE_EXIT
+	{
+		// Once finished, clear the async flags of the meshes and its subobjects
+		ClearAsyncFlags(CachedMesh.StaticMesh);
+		ClearAsyncFlags(CachedMesh.DynamicMesh);
+	};
 
-		if (InSharedParams.GlyphMeshParameters->MeshType == EText3DMeshType::Dynamic)
-		{
-			CachedMesh.DynamicMesh = NewObject<UDynamicMesh>(InSharedParams.TextSubsystem, NAME_None, RF_Transient | RF_Public);
-			OnMeshAllocated(CachedMesh.DynamicMesh);
-		}
+	CachedMesh.StaticMesh = NewObject<UStaticMesh>(InSharedParams.TextSubsystem, NAME_None, RF_Transient | RF_Public);
+	if (InSharedParams.GlyphMeshParameters->MeshType == EText3DMeshType::Dynamic)
+	{
+		CachedMesh.DynamicMesh = NewObject<UDynamicMesh>(InSharedParams.TextSubsystem, NAME_None, RF_Transient | RF_Public);
 	}
 
 	FText3DGlyphMeshBuilder MeshCreator;
diff --git a/Engine/Plugins/VirtualProduction/Text3D/Source/Text3D/Private/Subsystems/Text3DEngineSubsystem.h b/Engine/Plugins/VirtualProduction/Text3D/Source/Text3D/Private/Subsystems/Text3DEngineSubsystem.h
--- a/Engine/Plugins/VirtualProduction/Text3D/Source/Text3D/Private/Subsystems/Text3DEngineSubsystem.h
+++ b/Engine/Plugins/VirtualProduction/Text3D/Source/Text3D/Private/Subsystems/Text3DEngineSubsystem.h
@@ -110,7 +110,6 @@ public:
 	static UText3DEngineSubsystem* Get();
 
 	void ScheduleCacheCleanup();
-	void CancelCacheCleanup();
 
 	/** Returns the default material for text meshes */
 	UMaterial* GetDefaultMaterial() const;
diff --git a/Engine/Plugins/VirtualProduction/Text3D/Source/Text3D/Private/Text3DTypes.cpp b/Engine/Plugins/VirtualProduction/Text3D/Source/Text3D/Private/Text3DTypes.cpp
--- a/Engine/Plugins/VirtualProduction/Text3D/Source/Text3D/Private/Text3DTypes.cpp
+++ b/Engine/Plugins/VirtualProduction/Text3D/Source/Text3D/Private/Text3DTypes.cpp
@@ -39,7 +39,10 @@ namespace UE::Text3D::Geometry
 	{
 		if (FText3DCachedMesh* CachedMesh = Resolve())
 		{
-			CachedMesh->RefCount--;
+			if (CachedMesh->RefCount == 0 || --CachedMesh->RefCount == 0)
+			{
+				UText3DEngineSubsystem::Get()->ScheduleCacheCleanup();
+			}
 		}
 
 		FontFaceHash = 0;
diff --git a/Engine/Plugins/VirtualProduction/Text3D/Source/Text3D/Public/Settings/Text3DProjectSettings.h b/Engine/Plugins/VirtualProduction/Text3D/Source/Text3D/Public/Settings/Text3DProjectSettings.h
--- a/Engine/Plugins/VirtualProduction/Text3D/Source/Text3D/Public/Settings/Text3DProjectSettings.h
+++ b/Engine/Plugins/VirtualProduction/Text3D/Source/Text3D/Public/Settings/Text3DProjectSettings.h
@@ -72,12 +72,27 @@ public:
 
 	TEXT3D_API UMaterial* GetBaseMaterial(const FText3DMaterialKey& InKey) const;
 
+	UE_DEPRECATED(5.8, "Font face glyph fixed scheduled cleanup is deprecated. Glyphs are cleaned up on-demand")
 	TEXT3D_API void SetScheduleFontFaceGlyphCleanup(bool bInEnabled);
+
+	UE_DEPRECATED(5.8, "Font face glyph fixed scheduled cleanup is deprecated. Glyphs are cleaned up on-demand")
 	TEXT3D_API bool GetScheduleFontFaceGlyphCleanup() const;
 
+	UE_DEPRECATED(5.8, "Font face glyph fixed scheduled cleanup is deprecated. Glyphs are cleaned up on-demand")
 	TEXT3D_API void SetFontFaceGlyphCleanupPeriod(float InValue);
+
+	UE_DEPRECATED(5.8, "Font face glyph fixed scheduled cleanup is deprecated. Glyphs are cleaned up on-demand")
 	TEXT3D_API float GetFontFaceGlyphCleanupPeriod() const;
 
+	/** Sets the amount of time, in seconds, to wait before a cleanup happens after a mesh becomes unused  */
+	TEXT3D_API void SetFontFaceGlyphCleanupDelay(float InDelay);
+
+	/** The amount of time, in seconds, to wait before a cleanup happens after a mesh becomes unused  */
+	float GetFontFaceGlyphCleanupDelay() const
+	{
+		return FontFaceGlyphCleanupDelay;
+	}
+
 #if WITH_EDITOR
 	void OpenEditorSettingsWindow() const;
 	TEXT3D_API TConstArrayView<FString> GetFavoriteFonts() const;
@@ -96,7 +111,6 @@ public:
 
 protected:
 	void AddMaterial(FText3DMaterialKey&& InKey, TSoftObjectPtr<UMaterial>&& InMaterial);
-	void OnFontFaceCleanupOptionsChanged();
 
 #if WITH_EDITOR
 	//~ Begin UObject
@@ -126,13 +140,17 @@ protected:
 	UPROPERTY(Config, EditAnywhere, Category="Text3D")
 	FString FontDirectory;
 
-	/** Schedule font face glyph mesh cleanup */
-	UPROPERTY(Config, EditAnywhere, Category="Text3D")
-	bool bScheduleFontFaceGlyphCleanup = true;
+	UE_DEPRECATED(5.8, "Font face glyph fixed scheduled cleanup is deprecated. Glyphs are cleaned up on-demand")
+	UPROPERTY(meta=(DeprecatedProperty))
+	bool bScheduleFontFaceGlyphCleanup_DEPRECATED = true;
 
-	/** Interval for the font face glyph mesh cleanup */
-	UPROPERTY(Config, EditAnywhere, Category="Text3D", meta=(ClampMin="0", Units=Seconds, EditCondition="bScheduleFontFaceGlyphCleanup"))
-	float FontFaceGlyphCleanupPeriod = 300.f;
+	UE_DEPRECATED(5.8, "Font face glyph fixed scheduled cleanup is deprecated. Glyphs are cleaned up on-demand. Use FontFaceGlyphCleanupDelay instead")
+	UPROPERTY(meta=(DeprecatedProperty))
+	float FontFaceGlyphCleanupPeriod_DEPRECATED = 300.f;
+
+	/** Once a glyph mesh becomes unused, how long to wait before a cleanup happens. <= 0 means immediate. */
+	UPROPERTY(Config, EditAnywhere, Category="Text3D")
+	float FontFaceGlyphCleanupDelay = 5.f;
 
 #if WITH_EDITORONLY_DATA
 	/** Favorites fonts pinned in the font viewer dropdown */
