MotionDesign : Text3D
- Added component function to retrieve casted text renderer in blueprint, this will allow user to access functions exposed in the renderers

#Synced-CL 51025551

diff --git a/Engine/Plugins/VirtualProduction/Text3D/Source/Text3D/Private/Text3DComponent.cpp b/Engine/Plugins/VirtualProduction/Text3D/Source/Text3D/Private/Text3DComponent.cpp
--- a/Engine/Plugins/VirtualProduction/Text3D/Source/Text3D/Private/Text3DComponent.cpp
+++ b/Engine/Plugins/VirtualProduction/Text3D/Source/Text3D/Private/Text3DComponent.cpp
@@ -668,6 +668,16 @@ UText3DRendererBase* UText3DComponent::GetTextRenderer() const
 	return TextRenderer;
 }
 
+UText3DRendererBase* UText3DComponent::GetTextRenderer(TSubclassOf<UText3DRendererBase> InRendererClass) const
+{
+	if (!TextRenderer || !InRendererClass.Get())
+	{
+		return nullptr;
+	}
+	
+	return TextRenderer->IsA(InRendererClass.Get()) ? TextRenderer : nullptr;
+}
+
 UText3DTokenExtensionBase* UText3DComponent::GetTokenExtension(TSubclassOf<UText3DTokenExtensionBase> InExtensionClass)
 {
 	if (!TokenExtension || !InExtensionClass.Get())
diff --git a/Engine/Plugins/VirtualProduction/Text3D/Source/Text3D/Public/Text3DComponent.h b/Engine/Plugins/VirtualProduction/Text3D/Source/Text3D/Public/Text3DComponent.h
--- a/Engine/Plugins/VirtualProduction/Text3D/Source/Text3D/Public/Text3DComponent.h
+++ b/Engine/Plugins/VirtualProduction/Text3D/Source/Text3D/Public/Text3DComponent.h
@@ -282,6 +282,9 @@ public:
 	
 	TEXT3D_API UText3DRendererBase* GetTextRenderer() const;
 
+	UFUNCTION(BlueprintCallable, Category="Text3D", meta=(DeterminesOutputType="InRendererClass"))
+	UText3DRendererBase* GetTextRenderer(TSubclassOf<UText3DRendererBase> InRendererClass) const;
+
 	UFUNCTION(BlueprintCallable, Category="Text3D", meta=(DeterminesOutputType="InExtensionClass"))
 	UText3DTokenExtensionBase* GetTokenExtension(TSubclassOf<UText3DTokenExtensionBase> InExtensionClass);
 	
