MotionDesign : Text3D
- TextRenderer and TextExtension can now be referenced in a blueprint variable

#Synced-CL 51025551

diff --git a/Engine/Plugins/VirtualProduction/Text3D/Source/Text3D/Public/Extensions/Text3DExtensionBase.h b/Engine/Plugins/VirtualProduction/Text3D/Source/Text3D/Public/Extensions/Text3DExtensionBase.h
--- a/Engine/Plugins/VirtualProduction/Text3D/Source/Text3D/Public/Extensions/Text3DExtensionBase.h
+++ b/Engine/Plugins/VirtualProduction/Text3D/Source/Text3D/Public/Extensions/Text3DExtensionBase.h
@@ -22,7 +22,7 @@ enum class EText3DExtensionResult : uint8
  * Extensions are piece of data and logic needed to render Text3D that can be reused by multiple renderers,
  * they execute once during every renderer update at the right moment
  */
-UCLASS(MinimalAPI, Abstract)
+UCLASS(MinimalAPI, Abstract, BlueprintType)
 class UText3DExtensionBase : public UObject
 {
 	GENERATED_BODY()
diff --git a/Engine/Plugins/VirtualProduction/Text3D/Source/Text3D/Public/Renderers/Text3DRendererBase.h b/Engine/Plugins/VirtualProduction/Text3D/Source/Text3D/Public/Renderers/Text3DRendererBase.h
--- a/Engine/Plugins/VirtualProduction/Text3D/Source/Text3D/Public/Renderers/Text3DRendererBase.h
+++ b/Engine/Plugins/VirtualProduction/Text3D/Source/Text3D/Public/Renderers/Text3DRendererBase.h
@@ -13,7 +13,7 @@ class UText3DComponent;
  * Base class for a rendering implementation of Text3D
  * The whole rendering logic should be encapsulated into an instance of this class
  */
-UCLASS(MinimalAPI, Abstract)
+UCLASS(MinimalAPI, Abstract, BlueprintType)
 class UText3DRendererBase : public UObject
 {
 	GENERATED_BODY()
@@ -48,19 +48,19 @@ public:
 
 protected:
 	/** Create and setup the implementation components, called on load or creation */
-	virtual void OnCreate() PURE_VIRTUAL(UText3DRenderingImplementationBase::OnCreate, );
+	virtual void OnCreate() PURE_VIRTUAL(UText3DRendererBase::OnCreate, );
 
 	/** Update rendering state of text characters, called when render state is outdated */
-	virtual void OnUpdate(const UE::Text3D::Renderer::FUpdateParameters& InParameters) PURE_VIRTUAL(UText3DRenderingImplementationBase::OnUpdate, );
+	virtual void OnUpdate(const UE::Text3D::Renderer::FUpdateParameters& InParameters) PURE_VIRTUAL(UText3DRendererBase::OnUpdate, );
 
 	/** Clear rendering state and remove all visible characters */
-	virtual void OnClear() PURE_VIRTUAL(UText3DRenderingImplementationBase::OnClear, );
+	virtual void OnClear() PURE_VIRTUAL(UText3DRendererBase::OnClear, );
 
 	/** Destroy and clean the implementation components, called on destroy or deactivation */
-	virtual void OnDestroy() PURE_VIRTUAL(UText3DRenderingImplementationBase::OnDestroy, );
+	virtual void OnDestroy() PURE_VIRTUAL(UText3DRendererBase::OnDestroy, );
 
 	/** Calculate the bounds of the rendered text */
-	virtual FBox OnCalculateBounds() const PURE_VIRTUAL(UText3DRenderingImplementationBase::OnCalculateBounds, return FBox(ForceInitToZero); );
+	virtual FBox OnCalculateBounds() const PURE_VIRTUAL(UText3DRendererBase::OnCalculateBounds, return FBox(ForceInitToZero); );
 
 #if WITH_EDITOR
 	void BindDebugDelegate();
