Text3D: Introduced Text3D builder that is able to do both incremental and blocking builds.

#Synced-CL 51290016

diff --git a/Engine/Plugins/VirtualProduction/Text3D/Source/Text3D/Private/Characters/Text3DCharacterBase.cpp b/Engine/Plugins/VirtualProduction/Text3D/Source/Text3D/Private/Characters/Text3DCharacterBase.cpp
--- a/Engine/Plugins/VirtualProduction/Text3D/Source/Text3D/Private/Characters/Text3DCharacterBase.cpp
+++ b/Engine/Plugins/VirtualProduction/Text3D/Source/Text3D/Private/Characters/Text3DCharacterBase.cpp
@@ -171,7 +171,6 @@ void UText3DCharacterBase::OnCharacterDataChanged(EText3DRendererFlags InFlags)
 {
 	if (UText3DComponent* Component = GetTypedOuter<UText3DComponent>())
 	{
-		constexpr bool bImmediate = true;
-		Component->RequestUpdate(InFlags, bImmediate);
+		Component->RequestUpdate(InFlags);
 	}
 }
diff --git a/Engine/Plugins/VirtualProduction/Text3D/Source/Text3D/Private/Extensions/Text3DDefaultGeometryExtension.cpp b/Engine/Plugins/VirtualProduction/Text3D/Source/Text3D/Private/Extensions/Text3DDefaultGeometryExtension.cpp
--- a/Engine/Plugins/VirtualProduction/Text3D/Source/Text3D/Private/Extensions/Text3DDefaultGeometryExtension.cpp
+++ b/Engine/Plugins/VirtualProduction/Text3D/Source/Text3D/Private/Extensions/Text3DDefaultGeometryExtension.cpp
@@ -191,13 +191,15 @@ const FGlyphMeshParameters* UText3DDefaultGeometryExtension::GetGlyphMeshParamet
 	return &GlyphMeshParameters;
 }
 
-EText3DExtensionResult UText3DDefaultGeometryExtension::PreRendererUpdate(const UE::Text3D::Renderer::FUpdateParameters& InParameters)
+void UText3DDefaultGeometryExtension::PrepareBuild(EText3DRendererFlags InUpdateFlags)
 {
-	if (InParameters.CurrentFlag != EText3DRendererFlags::Geometry)
+	if (!EnumHasAnyFlags(InUpdateFlags, EText3DRendererFlags::Geometry))
 	{
-		return EText3DExtensionResult::Active;
+		return;
 	}
 
+	TRACE_CPUPROFILER_EVENT_SCOPE(UText3DDefaultGeometryExtension::PrepareRebuild);
+
 	const UText3DComponent* TextComponent = GetText3DComponent();
 	const int32 TypefaceIndex = TextComponent->GetTypeFaceIndex();
 
@@ -205,14 +207,14 @@ EText3DExtensionResult UText3DDefaultGeometryExtension::PreRendererUpdate(const
 	{
 		const UFont* Font = TextComponent->GetFont();
 		UE_LOG(LogText3D, Error, TEXT("Invalid typeface index returned for font %s : %i"), Font ? *Font->GetName() : TEXT("invalid"), TypefaceIndex);
-		return EText3DExtensionResult::Failed;
+		return;
 	}
 
 	const float GlyphExtrude = FMath::TruncToFloat(GetExtrude() * 10.f) / 10.f;
 	const float GlyphBevel = FMath::TruncToFloat(GetBevel() * 10.f) / 10.f;
 	const float GlyphOutline = FMath::TruncToFloat(GetOutline() * 10.f) / 10.f;
 
-	const UText3DRendererBase* TextRenderer = TextComponent->GetTextRenderer();
+	const UText3DRendererBase* TextRenderer = TextComponent->GetTextRendererClass().GetDefaultObject();
 	GlyphMeshParameters =
 	{
 		.Extrude = GlyphExtrude,
@@ -225,7 +227,10 @@ EText3DExtensionResult UText3DDefaultGeometryExtension::PreRendererUpdate(const
 		.PivotOffset = GetPivotOffset(),
 		.MeshType = TextRenderer ? TextRenderer->GetMeshType() : EText3DMeshType::Unknown,
 	};
+}
 
+EText3DExtensionResult UText3DDefaultGeometryExtension::PreRendererUpdate(const UE::Text3D::Renderer::FUpdateParameters& InParameters)
+{
 	// Keep to clean after
 	return EText3DExtensionResult::Active;
 }
diff --git a/Engine/Plugins/VirtualProduction/Text3D/Source/Text3D/Private/Extensions/Text3DDefaultLayoutExtension.cpp b/Engine/Plugins/VirtualProduction/Text3D/Source/Text3D/Private/Extensions/Text3DDefaultLayoutExtension.cpp
--- a/Engine/Plugins/VirtualProduction/Text3D/Source/Text3D/Private/Extensions/Text3DDefaultLayoutExtension.cpp
+++ b/Engine/Plugins/VirtualProduction/Text3D/Source/Text3D/Private/Extensions/Text3DDefaultLayoutExtension.cpp
@@ -207,17 +207,18 @@ FVector UText3DDefaultLayoutExtension::GetTextScale() const
 	return TextScale;
 }
 
-EText3DExtensionResult UText3DDefaultLayoutExtension::PreRendererUpdate(const UE::Text3D::Renderer::FUpdateParameters& InParameters)
+void UText3DDefaultLayoutExtension::PrepareBuild(EText3DRendererFlags InUpdateFlags)
 {
-	using namespace UE::Text3D::Layout;
-
-	TRACE_CPUPROFILER_EVENT_SCOPE(UText3DDefaultLayoutExtension::PreRendererUpdate);
-
-	if (InParameters.CurrentFlag != EText3DRendererFlags::Geometry && InParameters.CurrentFlag != EText3DRendererFlags::Layout)
+	// Currently only concerned about geometry / glyph building
+	if (!EnumHasAnyFlags(InUpdateFlags, EText3DRendererFlags::Geometry))
 	{
-		return EText3DExtensionResult::Active;
+		return;
 	}
 
+	TRACE_CPUPROFILER_EVENT_SCOPE(UText3DDefaultLayoutExtension::PrepareRebuild);
+
+	using namespace UE::Text3D::Layout;
+
 	const UText3DComponent* Text3DComponent = GetText3DComponent();
 	
 	if (!GlyphText.IsValid())
@@ -242,6 +243,68 @@ EText3DExtensionResult UText3DDefaultLayoutExtension::PreRendererUpdate(const UE
 	const UText3DStyleExtensionBase* FormatStyleExtension = Text3DComponent->GetStyleExtension();
 	FText3DLayoutShaper::Get()->ShapeBidirectionalText(FormatStyleExtension->GetCustomStyles(), *FormatStyleExtension->GetDefaultStyle(), Text3DComponent->GetFormattedText().ToString(), TextLayout, TextLayoutMarshaller, GlyphText->Lines);
 
+#if WITH_FREETYPE
+	UText3DEngineSubsystem* const TextSubsystem = UText3DEngineSubsystem::Get();
+	if (!ensure(TextSubsystem))
+	{
+		return;
+	}
+
+	const UText3DGeometryExtensionBase* const GeometryExtension = Text3DComponent->GetGeometryExtension();
+	if (!ensure(GeometryExtension))
+	{
+		return;
+	}
+
+	const FGlyphMeshParameters& GlyphMeshParameters = *GeometryExtension->GetGlyphMeshParameters();
+	TArray<FText3DBuildGlyphMeshDesc> BuildGlyphMeshDescs;
+
+	TRACE_CPUPROFILER_EVENT_SCOPE(UText3DDefaultLayoutExtension::PrepareCachedMeshes);
+	for (const TConstEnumerateRef<FGlyphLine> GlyphLine : EnumerateRange(GlyphText->Lines))
+	{
+		for (const TConstEnumerateRef<FGlyphEntry> GlyphEntry : EnumerateRange(GlyphLine->Glyphs))
+		{
+			if (!FGlyphText::IsValidGlyph(*GlyphEntry))
+			{
+				continue;
+			}
+			const TSharedPtr<FFreeTypeFace> FontFace = GlyphEntry->Entry.FontFaceData->FontFace.Pin();
+			if (FText3DFontFaceCache* FontFaceCache = TextSubsystem->FindOrAddCachedFontFace(FontFace))
+			{
+				if (TOptional<FText3DBuildGlyphMeshDesc> GlyphOutline = FontFaceCache->PrepareCachedMesh(GlyphEntry->Entry.GlyphIndex, GlyphMeshParameters))
+				{
+					BuildGlyphMeshDescs.Emplace(MoveTemp(*GlyphOutline));
+				}
+			}
+		}
+	}
+
+	TextSubsystem->QueueBuildGlyphMeshes(MoveTemp(BuildGlyphMeshDescs), GlyphMeshParameters);
+	BuildGlyphMeshDescs.Reset();
+#else
+	UE_LOG(LogText3D, Warning, TEXT("FreeType is not available, cannot proceed without it"));
+#endif
+}
+
+EText3DExtensionResult UText3DDefaultLayoutExtension::PreRendererUpdate(const UE::Text3D::Renderer::FUpdateParameters& InParameters)
+{
+	using namespace UE::Text3D::Layout;
+
+	TRACE_CPUPROFILER_EVENT_SCOPE(UText3DDefaultLayoutExtension::PreRendererUpdate);
+
+	if (!GlyphText.IsValid())
+	{
+		UE_LOG(LogText3D, Warning, TEXT("Text3DDefaultLayoutExtension::PreRendererUpdate failed. Invalid GlyphText"));
+		return EText3DExtensionResult::Failed;
+	}
+
+	if (InParameters.CurrentFlag != EText3DRendererFlags::Geometry && InParameters.CurrentFlag != EText3DRendererFlags::Layout)
+	{
+		return EText3DExtensionResult::Active;
+	}
+
+	const UText3DComponent* Text3DComponent = GetText3DComponent();
+
 	const int32 CharacterCount = Algo::TransformAccumulate(GlyphText->Lines, [](const FGlyphLine& InLine)
 	{
 		return Algo::CountIf(InLine.Glyphs, [](const FGlyphEntry& InGlyph)
@@ -310,80 +373,13 @@ EText3DExtensionResult UText3DDefaultLayoutExtension::PreRendererUpdate(const UE
 		return EText3DExtensionResult::Failed;
 	}
 
-	const FGlyphMeshParameters& GlyphMeshParameters = *GeometryExtension->GetGlyphMeshParameters();
-
-	CharacterIndex = 0;
-
 	if (EnumHasAnyFlags(InParameters.UpdateFlags, EText3DRendererFlags::Geometry))
 	{
-		struct FBuildGlyphMeshDesc
-		{
-			UE::Text3D::FGlyphOutline GlyphOutline;
-			TSharedPtr<FFreeTypeFace> FontFace;
-			uint32 GlyphIndex = 0;
-		};
-		TArray<FBuildGlyphMeshDesc> BuildGlyphMeshDescs;
-
-		{
-			TRACE_CPUPROFILER_EVENT_SCOPE(UText3DDefaultLayoutExtension::PrepareCachedMeshes);
-			for (const TConstEnumerateRef<FGlyphLine> GlyphLine : EnumerateRange(GlyphText->Lines))
-			{
-				for (const TConstEnumerateRef<FGlyphEntry> GlyphEntry : EnumerateRange(GlyphLine->Glyphs))
-				{
-					if (!FGlyphText::IsValidGlyph(*GlyphEntry))
-					{
-						continue;
-					}
-
-					if (!Text3DComponent->GetCharacter(CharacterIndex))
-					{
-						UE_LOG(LogText3D, Warning, TEXT("Invalid character object returned for index %i after layout calculations, cannot proceed"), CharacterIndex)
-						return EText3DExtensionResult::Failed;
-					}
-
-					const TSharedPtr<FFreeTypeFace>& FontFace = GlyphText->FontFaces[CharacterIndex++];
-					if (FText3DFontFaceCache* FontFaceCache = TextSubsystem->FindOrAddCachedFontFace(FontFace))
-					{
-						if (TOptional<UE::Text3D::FGlyphOutline> GlyphOutline = FontFaceCache->PrepareCachedMesh(GlyphEntry->Entry.GlyphIndex, GlyphMeshParameters))
-						{
-							BuildGlyphMeshDescs.Emplace(MoveTemp(*GlyphOutline), FontFace, GlyphEntry->Entry.GlyphIndex);
-						}
-					}
-					else
-					{
-						UE_LOG(LogText3D, Warning, TEXT("Failed to get cached font data for '%u %i' in Text3D geometry extension"), FText3DFontFaceCache::GetFontFaceHash(FontFace), GlyphEntry->Entry.GlyphIndex);
-						return EText3DExtensionResult::Failed;
-					}
-				}
-			}
-		}
-
-		TRACE_CPUPROFILER_EVENT_SCOPE(UText3DDefaultLayoutExtension::BuildGlyphMeshes);
-		const FText3DFontFaceCache::FBuildGlyphMeshSharedDesc SharedParams
-			{
-				.TextSubsystem = TextSubsystem,
-				.GlyphMeshParameters = &GlyphMeshParameters,
-				.DefaultMaterial = TextSubsystem->GetDefaultMaterial(),
-			};
-
-		ParallelFor(BuildGlyphMeshDescs.Num(), 
-			[&SharedParams, &BuildGlyphMeshDescs](int32 InTaskInfoIndex)
-			{
-				TRACE_CPUPROFILER_EVENT_SCOPE(UText3DDefaultLayoutExtension::BuildGlyphMesh);
-				const FBuildGlyphMeshDesc& BuildGlyphMeshDesc = BuildGlyphMeshDescs[InTaskInfoIndex];
-
-				if (FText3DFontFaceCache* FontFaceCache = SharedParams.TextSubsystem->FindCachedFontFace(BuildGlyphMeshDesc.FontFace))
-				{
-					const FText3DFontFaceCache::FBuildGlyphMeshInstanceDesc InstanceParams
-						{
-							.GlyphOutline = &BuildGlyphMeshDesc.GlyphOutline,
-							.GlyphIndex = BuildGlyphMeshDesc.GlyphIndex,
-						};
-					FontFaceCache->BuildGlyphMesh(SharedParams, InstanceParams);
-				}
-			});
+		TextSubsystem->ProcessBuildGlyphMeshes();	
 	}
 
+	const FGlyphMeshParameters& GlyphMeshParameters = *GeometryExtension->GetGlyphMeshParameters();
+
 	CharacterIndex = 0;
 
 	for (TConstEnumerateRef<FGlyphLine> GlyphLine : EnumerateRange(GlyphText->Lines))
diff --git a/Engine/Plugins/VirtualProduction/Text3D/Source/Text3D/Private/Extensions/Text3DDefaultStyleExtension.cpp b/Engine/Plugins/VirtualProduction/Text3D/Source/Text3D/Private/Extensions/Text3DDefaultStyleExtension.cpp
--- a/Engine/Plugins/VirtualProduction/Text3D/Source/Text3D/Private/Extensions/Text3DDefaultStyleExtension.cpp
+++ b/Engine/Plugins/VirtualProduction/Text3D/Source/Text3D/Private/Extensions/Text3DDefaultStyleExtension.cpp
@@ -94,11 +94,15 @@ UText3DStyleBase* UText3DDefaultStyleExtension::GetStyle(FName InStyleName) cons
 	return nullptr;
 }
 
-EText3DExtensionResult UText3DDefaultStyleExtension::PreRendererUpdate(const UE::Text3D::Renderer::FUpdateParameters& InParameters)
+void UText3DDefaultStyleExtension::PrepareBuild(EText3DRendererFlags InUpdateFlags)
 {
+	Super::PrepareBuild(InUpdateFlags);
+
+	TRACE_CPUPROFILER_EVENT_SCOPE(UText3DDefaultStyleExtension::PrepareRebuild);
+
 	if (!DefaultFontStyle.IsValid()
 		|| !TextStyleSet.IsValid()
-		|| EnumHasAnyFlags(InParameters.UpdateFlags, EText3DRendererFlags::Geometry))
+		|| EnumHasAnyFlags(InUpdateFlags, EText3DRendererFlags::Geometry))
 	{
 		const UText3DComponent* Text3DComponent = GetText3DComponent();
 		const UFont* Font = Text3DComponent->GetFont();
@@ -164,6 +168,10 @@ EText3DExtensionResult UText3DDefaultStyleExtension::PreRendererUpdate(const UE:
 			MaterialExtension->UnregisterMaterialOverride(MaterialName);
 		}
 	}
+}
+
+EText3DExtensionResult UText3DDefaultStyleExtension::PreRendererUpdate(const UE::Text3D::Renderer::FUpdateParameters& InParameters)
+{
 
 	return EText3DExtensionResult::Finished;
 }
diff --git a/Engine/Plugins/VirtualProduction/Text3D/Source/Text3D/Private/Extensions/Text3DDefaultTokenExtension.cpp b/Engine/Plugins/VirtualProduction/Text3D/Source/Text3D/Private/Extensions/Text3DDefaultTokenExtension.cpp
--- a/Engine/Plugins/VirtualProduction/Text3D/Source/Text3D/Private/Extensions/Text3DDefaultTokenExtension.cpp
+++ b/Engine/Plugins/VirtualProduction/Text3D/Source/Text3D/Private/Extensions/Text3DDefaultTokenExtension.cpp
@@ -36,8 +36,10 @@ const FText& UText3DDefaultTokenExtension::GetFormattedText() const
 	return FormattedText;
 }
 
-EText3DExtensionResult UText3DDefaultTokenExtension::PreRendererUpdate(const UE::Text3D::Renderer::FUpdateParameters& InParameters)
+void UText3DDefaultTokenExtension::PrepareBuild(EText3DRendererFlags InUpdateFlags)
 {
+	TRACE_CPUPROFILER_EVENT_SCOPE(UText3DDefaultTokenExtension::PrepareRebuild);
+
 	const UText3DComponent* Text3DComponent = GetText3DComponent();
 	FormattedText = Text3DComponent->GetText();
 
@@ -57,7 +59,10 @@ EText3DExtensionResult UText3DDefaultTokenExtension::PreRendererUpdate(const UE:
 
 	// For backward compatibility
 	Text3DComponent->FormatText(FormattedText);
+}
 
+EText3DExtensionResult UText3DDefaultTokenExtension::PreRendererUpdate(const UE::Text3D::Renderer::FUpdateParameters& InParameters)
+{
 	return EText3DExtensionResult::Finished;
 }
 
diff --git a/Engine/Plugins/VirtualProduction/Text3D/Source/Text3D/Private/Renderers/Text3DRendererBase.cpp b/Engine/Plugins/VirtualProduction/Text3D/Source/Text3D/Private/Renderers/Text3DRendererBase.cpp
--- a/Engine/Plugins/VirtualProduction/Text3D/Source/Text3D/Private/Renderers/Text3DRendererBase.cpp
+++ b/Engine/Plugins/VirtualProduction/Text3D/Source/Text3D/Private/Renderers/Text3DRendererBase.cpp
@@ -92,159 +92,12 @@ void UText3DRendererBase::Create()
 	UE_LOG(LogText3D, Verbose, TEXT("%s : Text3DRenderer %s Created"), !!Owner ? *Owner->GetActorNameOrLabel() : TEXT("Invalid owner"), *GetFriendlyName().ToString())
 }
 
-void UText3DRendererBase::Update(EText3DRendererFlags InFlags)
+void UText3DRendererBase::Update(const UE::Text3D::Renderer::FUpdateParameters& InParameters)
 {
-	using namespace UE::Text3D::Renderer;
-
-	if (!bInitialized)
-	{
-		return;
-	}
-
-	const UText3DComponent* Text3DComponent = GetText3DComponent();
-	if (!IsValid(Text3DComponent))
-	{
-		return;
-	}
-	
-	// Let's make sure text extension subobjects are valid and usable !
-	UText3DTokenExtensionBase* TokenExtension = Text3DComponent->GetTokenExtension();
-	if (!TokenExtension)
-	{
-		UE_LOG(LogText3D, Warning, TEXT("Invalid Text3D Token extension, cannot proceed !"))
-		return;
-	}
-
-	UText3DStyleExtensionBase* StyleExtension = Text3DComponent->GetStyleExtension();
-	if (!StyleExtension)
-	{
-		UE_LOG(LogText3D, Warning, TEXT("Invalid Text3D Style extension, cannot proceed !"))
-		return;
-	}
-
-	UText3DCharacterExtensionBase* CharacterExtension = Text3DComponent->GetCharacterExtension();
-	if (!CharacterExtension)
-	{
-		UE_LOG(LogText3D, Warning, TEXT("Invalid Text3D Character extension, cannot proceed !"))
-		return;
-	}
-
-	UText3DGeometryExtensionBase* GeometryExtension = Text3DComponent->GetGeometryExtension();
-	if (!GeometryExtension)
-	{
-		UE_LOG(LogText3D, Warning, TEXT("Invalid Text3D Geometry extension, cannot proceed !"))
-		return;
-	}
-
-	UText3DLayoutExtensionBase* LayoutExtension = Text3DComponent->GetLayoutExtension();
-	if (!LayoutExtension)
-	{
-		UE_LOG(LogText3D, Warning, TEXT("Invalid Text3D Layout extension, cannot proceed !"))
-		return;
-	}
-
-	UText3DMaterialExtensionBase* MaterialExtension = Text3DComponent->GetMaterialExtension();
-	if (!MaterialExtension)
-	{
-		UE_LOG(LogText3D, Warning, TEXT("Invalid Text3D Material extension, cannot proceed !"))
-		return;
-	}
-
-	UText3DRenderingExtensionBase* RenderingExtension = Text3DComponent->GetRenderingExtension();
-	if (!RenderingExtension)
-	{
-		UE_LOG(LogText3D, Warning, TEXT("Invalid Text3D Rendering extension, cannot proceed !"))
-		return;
-	}
-
-	TArray<UText3DExtensionBase*, TInlineAllocator<8>> Extensions
-	{
-		TokenExtension,
-		StyleExtension,
-		CharacterExtension,
-		GeometryExtension,
-		LayoutExtension,
-		MaterialExtension,
-		RenderingExtension
-	};
-
-	// Add optional extensions for layout effects
-	for (UText3DLayoutEffectBase* LayoutEffect : Text3DComponent->GetLayoutEffects())
-	{
-		if (LayoutEffect)
-		{
-			Extensions.Add(LayoutEffect);
-		}
-	}
-
-	// Sort in reverse order since we iterate in reverse order
-	Extensions.StableSort([](const UText3DExtensionBase& InExtensionA, const UText3DExtensionBase& InExtensionB)->bool
-	{
-		return InExtensionA.GetUpdatePriority() > InExtensionB.GetUpdatePriority();
-	});
-
-	uint8 Flag = 1 << 0;
-	uint8 FlagsLeft = static_cast<uint8>(InFlags);
-
-	while (Flag < static_cast<uint8>(EText3DRendererFlags::All))
+	if (bInitialized)
 	{
-		const EText3DRendererFlags FlagCasted = static_cast<EText3DRendererFlags>(Flag);
-
-		if (EnumHasAnyFlags(InFlags, FlagCasted))
-		{
-			FlagsLeft -= Flag;
-
-			const FUpdateParameters Parameters
-			{
-				.UpdateFlags = InFlags,
-				.CurrentFlag = FlagCasted,
-				.bIsLastFlag = FlagsLeft == 0
-			};
-
-			for (int32 Index = Extensions.Num() - 1; Index >= 0; --Index)
-			{
-				if (UText3DExtensionBase* Extension = Extensions[Index])
-				{
-					const EText3DExtensionResult ExtensionStatus = Extension->PreRendererUpdate(Parameters);
-
-					if (ExtensionStatus == EText3DExtensionResult::Failed)
-					{
-						UE_LOG(LogText3D, Error, TEXT("Failed to PRE update Text3D %s extension"), *GetNameSafe(Extension->GetClass()))
-						return;
-					}
-					else if (ExtensionStatus == EText3DExtensionResult::Finished)
-					{
-						Extensions.RemoveAt(Index);
-					}
-				}
-			}
-
-			OnUpdate(Parameters);
-
-			for (int32 Index = Extensions.Num() - 1; Index >= 0; --Index)
-			{
-				if (UText3DExtensionBase* Extension = Extensions[Index])
-				{
-					const EText3DExtensionResult ExtensionStatus = Extension->PostRendererUpdate(Parameters);
-
-					if (ExtensionStatus == EText3DExtensionResult::Failed)
-					{
-						UE_LOG(LogText3D, Error, TEXT("Failed to POST update Text3D %s extension"), *GetNameSafe(Extension->GetClass()))
-						return;
-					}
-					else if (ExtensionStatus == EText3DExtensionResult::Finished)
-					{
-						Extensions.RemoveAt(Index);
-					}
-				}
-			}
-		}
-
-		Flag <<= 1;
+		OnUpdate(InParameters);
 	}
-
-	const AActor* Owner = Text3DComponent->GetOwner();
-	UE_LOG(LogText3D, Verbose, TEXT("%s : Text3DRenderer %s Updated with flags %i"), !!Owner ? *Owner->GetActorNameOrLabel() : TEXT("Invalid owner"), *GetFriendlyName().ToString(), InFlags)
 }
 
 void UText3DRendererBase::Clear()
diff --git a/Engine/Plugins/VirtualProduction/Text3D/Source/Text3D/Private/Subsystems/Text3DBuildSystem.cpp b/Engine/Plugins/VirtualProduction/Text3D/Source/Text3D/Private/Subsystems/Text3DBuildSystem.cpp
new file mode 100644
--- /dev/null
+++ b/Engine/Plugins/VirtualProduction/Text3D/Source/Text3D/Private/Subsystems/Text3DBuildSystem.cpp
@@ -0,0 +1,283 @@
+// Copyright Epic Games, Inc. All Rights Reserved.
+
+#include "Text3DBuildSystem.h"
+#include "Engine/World.h"
+#include "Logs/Text3DLogs.h"
+#include "Misc/CoreDelegates.h"
+#include "Text3DComponent.h"
+
+namespace UE::Text3D::Private
+{
+TAutoConsoleVariable<float> CVarFrameTimeBudget(TEXT("Text3D.Build.FrameTimeBudget")
+	, 0.016f
+	, TEXT("The frame time budget to build Text3D components.\n")
+	  TEXT("If the frame is already past this, the build will be skipped unless it has starved for a while.\n")
+	  TEXT("See Text3D.Build.StarveFrameLimit")
+	);
+
+TAutoConsoleVariable<int32> CVarStarveFrameLimit(TEXT("Text3D.Build.StarveFrameLimit")
+	, 10
+	, TEXT("The maximum number of frames that the build can starve before forcing an incremental build"));
+
+TAutoConsoleVariable<bool> CVarEnableBuildTimeBudget(TEXT("Text3D.Build.EnableBuildTimeBudget")
+	, true
+	, TEXT("Whether to have time budget to build Text3D components."));
+
+TAutoConsoleVariable<float> CVarBuildTimeBudget(TEXT("Text3D.Build.BuildTimeBudget")
+	, 0.002f
+	, TEXT("The time budget to build Text3D components."));
+
+} // UE::Text3D::Private
+
+UText3DBuildSystem* UText3DBuildSystem::Get(TNotNull<const UText3DComponent*> InComponent)
+{
+	if (const UWorld* const World = InComponent->GetWorld())
+	{
+		return World->GetSubsystem<UText3DBuildSystem>();
+	}
+	return nullptr;
+}
+
+void UText3DBuildSystem::RequestTextBuild(TNotNull<UText3DComponent*> InComponent, EText3DRendererFlags InFlags, bool bInBlockingBuild)
+{
+	if (InFlags == EText3DRendererFlags::None)
+	{
+		return;
+	}
+
+	auto CanUseBuilder = 
+		[InFlags, bInBlockingBuild, InComponent](const UE::Text3D::FTextBuilder& InBuilder)->bool
+		{
+			// matching component
+			return InBuilder.GetComponent() == InComponent
+				// not currently building
+				&& !InBuilder.IsInProgress()
+				// allow if build is incremental. If blocking ensure the flags are exactly equal. If there's another flag, it can't be used.
+				&& (!bInBlockingBuild || InBuilder.GetFlags() == InFlags);
+		};
+
+	// Check in active builders. There could still be builders that haven't yet started where we could queue in our request.
+	UE::Text3D::FTextBuilder* Builder = nullptr;
+	for (UE::Text3D::FTextBuilder& CandidateBuilder : TextBuilders)
+	{
+		if (CanUseBuilder(CandidateBuilder) && EnumHasAllFlags(CandidateBuilder.GetFlags(), InFlags))
+		{
+			Builder = &CandidateBuilder;
+			break;
+		}
+	}
+
+	// Check in the queued builders.
+	if (!Builder)
+	{
+		for (UE::Text3D::FTextBuilder& CandidateBuilder : QueuedBuilders)
+		{
+			if (CanUseBuilder(CandidateBuilder))
+			{
+				Builder = &CandidateBuilder;
+				break;
+			}
+		}
+	}
+
+	if (!Builder)
+	{
+		Builder = &QueuedBuilders.Emplace_GetRef(InComponent);
+	}
+	Builder->PrepareFlags(InFlags);
+
+	ConditionallyRegisterDelegates();
+
+	if (bInBlockingBuild)
+	{
+		const uint64 StartCycles = FPlatformTime::Cycles64();
+		Builder->BlockingBuild();
+		const uint64 EndCycles = FPlatformTime::Cycles64();
+
+		// Accumulate the deficit so as not to impact next tick due to multiple blocking builds
+		BudgetDeficit += (EndCycles - StartCycles);
+	}
+}
+
+void UText3DBuildSystem::RemoveTextBuild(TNotNull<UText3DComponent*> InComponent)
+{
+	for (int32 Index = QueuedBuilders.Num() - 1; Index >= 0; --Index)
+	{
+		if (QueuedBuilders[Index].GetComponent() == InComponent)
+		{
+			QueuedBuilders.RemoveAt(Index);
+		}
+	}
+	for (int32 Index = TextBuilders.Num() - 1; Index >= 0; --Index)
+	{
+		if (TextBuilders[Index].GetComponent() == InComponent)
+		{
+			TextBuilders.RemoveAt(Index);
+		}
+	}
+}
+
+void UText3DBuildSystem::IncrementalTextBuild()
+{
+	using namespace UE::Text3D;
+
+	TRACE_CPUPROFILER_EVENT_SCOPE(UText3DWorldSubsystem::TextIncrementalUpdate);
+
+	// This was called within the same frame of the request. Skip to next frame to gather the frame duration
+	if (!BeginFrameCycles.IsSet())
+	{
+		return;
+	}
+
+	ON_SCOPE_EXIT
+	{
+		// Invalidate the gotten frame cycles
+		BeginFrameCycles.Reset();
+	};
+
+	TextBuilders.Append(MoveTemp(QueuedBuilders));
+	QueuedBuilders.Reset();
+
+	// Remove all finished builders
+	TextBuilders.RemoveAll(
+		[](const UE::Text3D::FTextBuilder& InBuilder)
+		{
+			return InBuilder.IsFinished();
+		});
+
+	if (TextBuilders.IsEmpty())
+	{
+		FCoreDelegates::OnEndFrame.Remove(EndFrameHandle);
+		EndFrameHandle.Reset();
+		return; // nothing to update
+	}
+
+	const uint64 CurrentCycles = FPlatformTime::Cycles64();
+
+	// If the incremental build has starved for too long, force it to happen
+	const bool bForceBuild = FramesStarved >= Private::CVarStarveFrameLimit.GetValueOnGameThread();
+
+	TOptional<uint64> BudgetEnd;
+	if (UE::Text3D::Private::CVarEnableBuildTimeBudget.GetValueOnGameThread())
+	{
+		BudgetEnd = CurrentCycles + FPlatformTime::SecondsToCycles64(UE::Text3D::Private::CVarBuildTimeBudget.GetValueOnGameThread());
+		// Only consider the budget deficit if not forcing a build
+		if (!bForceBuild)
+		{
+			if (*BudgetEnd <= BudgetDeficit)
+			{
+				++FramesStarved;
+				BudgetDeficit = 0;
+				UE_LOG(LogText3D, Verbose, TEXT("UText3DBuildSystem::IncrementalTextBuild ran out of build time budget due to blocking builds."));
+				return;
+			}
+			*BudgetEnd -= BudgetDeficit;
+		}
+	}
+	BudgetDeficit = 0;
+
+	// Skip if the frame duration has already gone for too long
+	const uint64 FrameDuration = CurrentCycles - *BeginFrameCycles;
+	if (!bForceBuild && FrameDuration >= FPlatformTime::SecondsToCycles64(Private::CVarFrameTimeBudget.GetValueOnGameThread()))
+	{
+		++FramesStarved;
+		UE_LOG(LogText3D, Verbose, TEXT("UText3DBuildSystem::IncrementalTextBuild ran out of frame time budget."));
+		return;
+	}
+	FramesStarved = 0;
+
+	// Update the builders that are still in the preparation phases
+	for (UE::Text3D::FTextBuilder& Builder : TextBuilders)
+	{
+		const EBuildResult Result = Builder.IncrementalBuild(BudgetEnd, /*MaxPhase*/UE::Text3D::EBuildPhase::PrepareExtensions);
+		if (Result == EBuildResult::OutOfBudget)
+		{
+			UE_LOG(LogText3D, VeryVerbose, TEXT("UText3DBuildSystem::IncrementalTextBuild ran out of budget while building. Continuing next frame"));
+			return; // ran out of budget
+		}
+	}
+
+	// Update all the builders with no phase limit
+	for (UE::Text3D::FTextBuilder& Builder : TextBuilders)
+	{
+		const EBuildResult Result = Builder.IncrementalBuild(BudgetEnd);
+		if (Result == EBuildResult::OutOfBudget)
+		{
+			UE_LOG(LogText3D, VeryVerbose, TEXT("UText3DBuildSystem::IncrementalTextBuild ran out of budget while building. Continuing next frame"));
+			return; // ran out of budget
+		}
+	}
+}
+
+void UText3DBuildSystem::Initialize(FSubsystemCollectionBase& InCollection)
+{
+	Super::Initialize(InCollection);
+
+#if WITH_EDITOR
+	FCoreUObjectDelegates::OnObjectsReplaced.AddWeakLambda(this, 
+		[this](const TMap<UObject*, UObject*>& InReplacementMap)
+		{
+			for (UE::Text3D::FTextBuilder& TextBuilder : QueuedBuilders)
+			{
+				TextBuilder.OnObjectsReplaced(InReplacementMap);	
+			}
+			for (UE::Text3D::FTextBuilder& TextBuilder : TextBuilders)
+			{
+				TextBuilder.OnObjectsReplaced(InReplacementMap);	
+			}
+		});
+#endif
+}
+
+void UText3DBuildSystem::Deinitialize()
+{
+	Super::Deinitialize();
+#if WITH_EDITOR
+	FCoreUObjectDelegates::OnObjectsReplaced.RemoveAll(this);
+#endif
+	UnregisterDelegates();
+}
+
+bool UText3DBuildSystem::DoesSupportWorldType(const EWorldType::Type InWorldType) const
+{
+	return InWorldType == EWorldType::Game
+		|| InWorldType == EWorldType::Editor
+		|| InWorldType == EWorldType::PIE
+		|| InWorldType == EWorldType::GamePreview;
+}
+
+bool UText3DBuildSystem::IsPendingBuildCompletion() const
+{
+	// Pending build if there are queued builders, or if there's at least one active builder that hasn't finished.
+	return !QueuedBuilders.IsEmpty() || TextBuilders.ContainsByPredicate(
+		[](const UE::Text3D::FTextBuilder& InTextBuilder)
+		{
+			return !InTextBuilder.IsFinished();
+		});
+}
+
+void UText3DBuildSystem::ConditionallyRegisterDelegates()
+{
+	if (!BeginFrameHandle.IsValid())
+	{
+		FramesStarved = 0;
+		BeginFrameHandle = FCoreDelegates::OnBeginFrame.AddWeakLambda(this, [this]
+			{
+				BeginFrameCycles = FPlatformTime::Cycles64();
+			});
+	}
+
+	if (!EndFrameHandle.IsValid())
+	{
+		EndFrameHandle = FCoreDelegates::OnEndFrame.AddUObject(this, &UText3DBuildSystem::IncrementalTextBuild);
+	}
+}
+
+void UText3DBuildSystem::UnregisterDelegates()
+{
+	FCoreDelegates::OnBeginFrame.Remove(BeginFrameHandle);
+	BeginFrameHandle.Reset();
+
+	FCoreDelegates::OnEndFrame.Remove(EndFrameHandle);
+	EndFrameHandle.Reset();
+}
diff --git a/Engine/Plugins/VirtualProduction/Text3D/Source/Text3D/Private/Subsystems/Text3DBuildSystem.h b/Engine/Plugins/VirtualProduction/Text3D/Source/Text3D/Private/Subsystems/Text3DBuildSystem.h
new file mode 100644
--- /dev/null
+++ b/Engine/Plugins/VirtualProduction/Text3D/Source/Text3D/Private/Subsystems/Text3DBuildSystem.h
@@ -0,0 +1,81 @@
+// Copyright Epic Games, Inc. All Rights Reserved.
+
+#pragma once
+
+#include "Delegates/IDelegateInstance.h"
+#include "IText3DBuildSystem.h"
+#include "Subsystems/WorldSubsystem.h"
+#include "Text3DBuilder.h"
+#include "Text3DBuildSystem.generated.h"
+
+class UText3DComponent;
+enum class EText3DRendererFlags : uint8;
+
+/**
+ * Manages all the text builds for a world 
+ * It's different from UText3DEngineSubsystem to keep these text builders referencing the text components scoped within the same world.
+ */
+UCLASS()
+class UText3DBuildSystem : public UWorldSubsystem, public IText3DBuildSystemInterface
+{
+	GENERATED_BODY()
+
+public:
+	static UText3DBuildSystem* Get(TNotNull<const UText3DComponent*> InComponent);
+
+	/** 
+	 * Adds the text component to the queue of builders 
+	 * @param InComponent the component to request build for
+	 * @param InFlags the renderer flags for the build
+	 * @param bInBlockingBuild whether the build should be blocking (immediate) or incremental (deferred, across multiple frames)
+	 */
+	void RequestTextBuild(TNotNull<UText3DComponent*> InComponent, EText3DRendererFlags InFlags, bool bInBlockingBuild);
+
+	/** Removes the builders that relate to the given component*/
+	void RemoveTextBuild(TNotNull<UText3DComponent*> InComponent);
+
+	/** Called on tick to process builders and call IncrementalBuild on them */
+	void IncrementalTextBuild();
+
+protected:
+	//~ Begin USubsystem
+	virtual void Initialize(FSubsystemCollectionBase& InCollection) override;
+	virtual void Deinitialize() override;
+	//~ End USubsystem
+
+	//~ Begin UWorldSubsystem
+	virtual bool DoesSupportWorldType(const EWorldType::Type InWorldType) const override;
+	//~ End UWorldSubsystem
+
+	//~ Begin IText3DBuildSystemInterface
+	virtual bool IsPendingBuildCompletion() const override;
+	//~ End IText3DBuildSystemInterface
+
+private:
+	/** Registers the build delegates only if they haven't been registered yet */
+	void ConditionallyRegisterDelegates();
+
+	/** Unregisters the build delegates */
+	void UnregisterDelegates();
+
+	/** Builders that are not yet processed */
+	TArray<UE::Text3D::FTextBuilder> QueuedBuilders;
+
+	/** Active builders */
+	TArray<UE::Text3D::FTextBuilder> TextBuilders;
+
+	/** Cycles measured in begin frame */
+	TOptional<uint64> BeginFrameCycles;
+
+	/** Accumulated cycles spent on blocking updates to determine if next tick an update should happen */
+	uint64 BudgetDeficit = 0;
+
+	/** Number of frames the incremental build has starved for */
+	int32 FramesStarved = 0;
+
+	/** Handle to the begin frame call to keep track how long this frame's been going for */
+	FDelegateHandle BeginFrameHandle;
+
+	/** Handle to the end frame call for incremental builds */
+	FDelegateHandle EndFrameHandle;
+};
diff --git a/Engine/Plugins/VirtualProduction/Text3D/Source/Text3D/Private/Subsystems/Text3DEngineSubsystem.cpp b/Engine/Plugins/VirtualProduction/Text3D/Source/Text3D/Private/Subsystems/Text3DEngineSubsystem.cpp
--- a/Engine/Plugins/VirtualProduction/Text3D/Source/Text3D/Private/Subsystems/Text3DEngineSubsystem.cpp
+++ b/Engine/Plugins/VirtualProduction/Text3D/Source/Text3D/Private/Subsystems/Text3DEngineSubsystem.cpp
@@ -14,6 +14,7 @@
 #include "Misc/FileHelper.h"
 #include "Settings/Text3DProjectSettings.h"
 #include "UDynamicMesh.h"
+#include "Async/ParallelFor.h"
 #include "UObject/ConstructorHelpers.h"
 
 #if WITH_FREETYPE
@@ -132,6 +133,12 @@ UMaterial* UText3DEngineSubsystem::GetDefaultMaterial() const
 
 bool UText3DEngineSubsystem::CleanupCache(float)
 {
+	if (bGlyphMeshBuildInProgress)
+	{
+		// Wait until next tick to see if cache can be cleaned up, glyph mesh building in progress
+		return true;
+	}
+
 	// Removing CleanupHandle Ticker while still executing this ticker delegate is by-design valid.
 	// see TickerTests.cpp:112 "a delegate removal from inside the delegate execution (used to be a deadlock)"
 	FTSTicker::RemoveTicker(CacheCleanupDelegate);
@@ -181,6 +188,63 @@ FText3DFontFaceCache* UText3DEngineSubsystem::FindCachedFontFace(uint32 InFontFa
 	return CachedFontFaces.Find(InFontFaceHash);
 }
 
+void UText3DEngineSubsystem::QueueBuildGlyphMeshes(TArray<FText3DBuildGlyphMeshDesc>&& InMeshDescs, const FGlyphMeshParameters& InGlyphMeshParameters)
+{
+	if (InMeshDescs.IsEmpty())
+	{
+		return;
+	}
+
+	const int32 GlyphParametersIndex = QueuedGlyphMeshParameters.Add(InGlyphMeshParameters);
+	for (FText3DBuildGlyphMeshDesc& MeshDesc : InMeshDescs)
+	{
+		MeshDesc.GlyphParametersIndex = GlyphParametersIndex;
+	}
+	QueuedMeshDescs.Append(MoveTemp(InMeshDescs));
+}
+
+void UText3DEngineSubsystem::ProcessBuildGlyphMeshes()
+{
+	// If queued mesh desc is empty, the mesh parameters should also be.
+	if (QueuedMeshDescs.IsEmpty() && ensure(QueuedGlyphMeshParameters.IsEmpty()))
+	{
+		return;
+	}
+
+	TGuardValue<bool> BuildScope(bGlyphMeshBuildInProgress, true);
+
+	const TArray<FGlyphMeshParameters> GlyphParameters = MoveTemp(QueuedGlyphMeshParameters);
+	QueuedGlyphMeshParameters.Reset();
+
+	const TArray<FText3DBuildGlyphMeshDesc> MeshDescs = MoveTemp(QueuedMeshDescs);
+	QueuedMeshDescs.Reset();
+
+	TRACE_CPUPROFILER_EVENT_SCOPE(UText3DEngineSubsystem::BuildGlyphMeshes);
+	const FText3DFontFaceCache::FBuildGlyphMeshSharedDesc SharedParams
+		{
+			.TextSubsystem = this,
+			.DefaultMaterial = GetDefaultMaterial(),
+		};
+
+	ParallelFor(MeshDescs.Num(), 
+		[&SharedParams, &MeshDescs, &GlyphParameters](int32 InMeshDescIndex)
+		{
+			TRACE_CPUPROFILER_EVENT_SCOPE(UText3DEngineSubsystem::BuildGlyphMesh);
+			const FText3DBuildGlyphMeshDesc& MeshDesc = MeshDescs[InMeshDescIndex];
+
+			if (FText3DFontFaceCache* FontFaceCache = SharedParams.TextSubsystem->FindCachedFontFace(MeshDesc.FontFace))
+			{
+				const FText3DFontFaceCache::FBuildGlyphMeshInstanceDesc InstanceParams
+					{
+						.GlyphOutline = &MeshDesc.GlyphOutline,
+						.GlyphMeshParameters = &GlyphParameters[MeshDesc.GlyphParametersIndex],
+						.GlyphIndex = MeshDesc.GlyphIndex,
+					};
+				FontFaceCache->BuildGlyphMesh(SharedParams, InstanceParams);
+			}
+		});
+}
+
 #if WITH_FREETYPE
 FText3DFontFaceCache* UText3DEngineSubsystem::FindOrAddCachedFontFace(const TSharedPtr<FFreeTypeFace>& InFontFace)
 {
@@ -248,7 +312,7 @@ void FText3DFontFaceCache::UpdateFontFace(const TSharedPtr<FFreeTypeFace>& InFon
 }
 #endif
 
-TOptional<UE::Text3D::FGlyphOutline> FText3DFontFaceCache::PrepareCachedMesh(uint32 InGlyphIndex, const FGlyphMeshParameters& InParameters)
+TOptional<FText3DBuildGlyphMeshDesc> FText3DFontFaceCache::PrepareCachedMesh(uint32 InGlyphIndex, const FGlyphMeshParameters& InParameters)
 {
 	const uint32 HashParameters = GetGlyphMeshHash(InGlyphIndex, InParameters);
 	if (FText3DCachedMesh* CachedMesh = GlyphMeshes.Find(HashParameters))
@@ -272,7 +336,12 @@ TOptional<UE::Text3D::FGlyphOutline> FText3DFontFaceCache::PrepareCachedMesh(uin
 	}
 #endif
 
-	return GlyphOutline;
+	FText3DBuildGlyphMeshDesc MeshDesc;
+	MeshDesc.GlyphOutline = MoveTemp(*GlyphOutline);
+	MeshDesc.FontFace = FontFace;
+	MeshDesc.GlyphIndex = InGlyphIndex;
+	MeshDesc.GlyphHandle = UE::Text3D::Geometry::FCachedFontFaceGlyphHandle(*this, InGlyphIndex, InParameters);
+	return MeshDesc;
 }
 
 uint32 FText3DFontFaceCache::GetGlyphMeshHash(uint32 InGlyphIndex, const FGlyphMeshParameters& InParameters)
@@ -293,7 +362,7 @@ void FText3DFontFaceCache::BuildGlyphMesh(const FBuildGlyphMeshSharedDesc& InSha
 		return;
 	}
 
-	const FGlyphMeshParameters& GlyphMeshParameters = *InSharedParams.GlyphMeshParameters;
+	const FGlyphMeshParameters& GlyphMeshParameters = *InInstanceParams.GlyphMeshParameters;
 
 	const uint32 HashParameters = GetGlyphMeshHash(InInstanceParams.GlyphIndex, GlyphMeshParameters);
 	FText3DCachedMesh& CachedMesh = GlyphMeshes.FindChecked(HashParameters);
@@ -335,7 +404,7 @@ void FText3DFontFaceCache::BuildGlyphMesh(const FBuildGlyphMeshSharedDesc& InSha
 	};
 
 	CachedMesh.StaticMesh = NewObject<UStaticMesh>(InSharedParams.TextSubsystem, NAME_None, RF_Transient | RF_Public);
-	if (InSharedParams.GlyphMeshParameters->MeshType == EText3DMeshType::Dynamic)
+	if (GlyphMeshParameters.MeshType == EText3DMeshType::Dynamic)
 	{
 		CachedMesh.DynamicMesh = NewObject<UDynamicMesh>(InSharedParams.TextSubsystem, NAME_None, RF_Transient | RF_Public);
 	}
diff --git a/Engine/Plugins/VirtualProduction/Text3D/Source/Text3D/Private/Subsystems/Text3DEngineSubsystem.h b/Engine/Plugins/VirtualProduction/Text3D/Source/Text3D/Private/Subsystems/Text3DEngineSubsystem.h
--- a/Engine/Plugins/VirtualProduction/Text3D/Source/Text3D/Private/Subsystems/Text3DEngineSubsystem.h
+++ b/Engine/Plugins/VirtualProduction/Text3D/Source/Text3D/Private/Subsystems/Text3DEngineSubsystem.h
@@ -3,21 +3,31 @@
 #pragma once
 
 #include "Containers/Ticker.h"
-#include "Text3D/Private/GeometryBuilders/Text3DGlyphContourNode.h"
+#include "GeometryBuilders/Text3DGlyphOutline.h"
 #include "Subsystems/EngineSubsystem.h"
 #include "Text3DTypes.h"
 #include "Text3DEngineSubsystem.generated.h"
 
-namespace UE::Text3D
-{
-	struct FGlyphOutline;
-}
-
 class FFreeTypeFace;
 class UMaterial;
 class UText3DEngineSubsystem;
 struct FSharedStruct;
 
+/** Struct containing the necessary information to build glyphs */
+struct FText3DBuildGlyphMeshDesc
+{
+	/** Generated outline to build from */
+	UE::Text3D::FGlyphOutline GlyphOutline;
+	/** Shared font face */
+	TSharedPtr<FFreeTypeFace> FontFace;
+	/** Index to the glyph to build */
+	uint32 GlyphIndex = 0;
+	/** Holds ref count of the glyph to prevent destruction */
+	UE::Text3D::Geometry::FCachedFontFaceGlyphHandle GlyphHandle;
+	/** Set when queuing the mesh desc to map to the glyph parameters used */
+	int32 GlyphParametersIndex = 0;
+};
+
 USTRUCT()
 struct FText3DFontFaceCache
 {
@@ -39,19 +49,19 @@ struct FText3DFontFaceCache
 	 * Loads the glyph contours and allocates the cached meshes the BuildGlyphMesh will be using (across multiple threads)
 	 * This happens in the game thread.
 	 */
-	TOptional<UE::Text3D::FGlyphOutline> PrepareCachedMesh(uint32 InGlyphIndex, const FGlyphMeshParameters& InParameters);
+	TOptional<FText3DBuildGlyphMeshDesc> PrepareCachedMesh(uint32 InGlyphIndex, const FGlyphMeshParameters& InParameters);
 
 	/** Data that does not change for a batch build */
 	struct FBuildGlyphMeshSharedDesc
 	{
 		TNotNull<UText3DEngineSubsystem*> TextSubsystem;
-		TNotNull<const FGlyphMeshParameters*> GlyphMeshParameters;
 		UMaterialInterface* DefaultMaterial;
 	};
 	/** Data that changes per build call in a batch */
 	struct FBuildGlyphMeshInstanceDesc
 	{
 		TNotNull<const UE::Text3D::FGlyphOutline*> GlyphOutline;
+		TNotNull<const FGlyphMeshParameters*> GlyphMeshParameters;
 		uint32 GlyphIndex;
 	};
 	/** 
@@ -121,6 +131,12 @@ public:
 
 	FText3DFontFaceCache* FindCachedFontFace(uint32 InFontFaceHash);
 
+	/** Adds the given mesh desc and glyph parameters to the queue */
+	void QueueBuildGlyphMeshes(TArray<FText3DBuildGlyphMeshDesc>&& InMeshDescs, const FGlyphMeshParameters& InGlyphMeshParameters);
+
+	/** Processes the queued mesh descs and glyph parameters across multiple threads */
+	void ProcessBuildGlyphMeshes();
+
 private:
 	// ~Begin UEngineSubsystem
 	virtual void Initialize(FSubsystemCollectionBase& InCollection) override;
@@ -145,9 +161,18 @@ private:
 	UPROPERTY()
 	TMap<uint32, FText3DFontFaceCache> CachedFontFaces;
 
+	/** Queued mesh descs to build next */
+	TArray<FText3DBuildGlyphMeshDesc> QueuedMeshDescs;
+
+	/** Glyph mesh parameters that the queued mesh descs reference */
+	TArray<FGlyphMeshParameters> QueuedGlyphMeshParameters;
+
 	/** the default material to use for text meshes */
 	UPROPERTY()
 	TObjectPtr<UMaterial> CachedDefaultMaterial;
 
 	FTSTicker::FDelegateHandle CacheCleanupDelegate;
+
+	/** Flag to prevent cleaning while building glyph meshes */
+	bool bGlyphMeshBuildInProgress = false;
 };
diff --git a/Engine/Plugins/VirtualProduction/Text3D/Source/Text3D/Private/Text3DBuilder.cpp b/Engine/Plugins/VirtualProduction/Text3D/Source/Text3D/Private/Text3DBuilder.cpp
new file mode 100644
--- /dev/null
+++ b/Engine/Plugins/VirtualProduction/Text3D/Source/Text3D/Private/Text3DBuilder.cpp
@@ -0,0 +1,352 @@
+// Copyright Epic Games, Inc. All Rights Reserved.
+
+#include "Text3DBuilder.h"
+#include "Extensions/Text3DExtensionBase.h"
+#include "Logs/Text3DLogs.h"
+#include "Renderers/Text3DRendererBase.h"
+#include "Text3DComponent.h"
+
+namespace UE::Text3D
+{
+
+FTextBuilder::FTextBuilder(UText3DComponent* InComponent)
+	: Component(InComponent)
+{
+}
+
+bool FTextBuilder::PrepareFlags(EText3DRendererFlags InFlags)
+{
+	if (IsInProgress())
+	{
+		return false;
+	}
+	EnumAddFlags(Flags, InFlags);
+	return true;
+}
+
+bool FTextBuilder::IsInProgress() const
+{
+	return Phase != EBuildPhase::NotStarted 
+		&& Phase != EBuildPhase::Finished;
+}
+
+bool FTextBuilder::IsFinished() const
+{
+	return Phase == EBuildPhase::Finished;
+}
+
+void FTextBuilder::BlockingBuild()
+{
+	TRACE_CPUPROFILER_EVENT_SCOPE(Text3D::BlockingBuild);
+
+	BudgetEnd.Reset();
+	TempPhaseLimit.Reset();
+	BuildInternal();
+}
+
+EBuildResult FTextBuilder::IncrementalBuild(TOptional<uint64> InBudgetEnd, TOptional<EBuildPhase> InTempMaxUpdatePhase)
+{
+	if (IsFinished())
+	{
+		return EBuildResult::BuildCompleted;
+	}
+
+	if (Flags == EText3DRendererFlags::None)
+	{
+		Finish();
+		return EBuildResult::BuildCompleted;
+	}
+
+	TRACE_CPUPROFILER_EVENT_SCOPE(Text3D::IncrementalBuild);
+
+	BudgetEnd = InBudgetEnd;
+	TempPhaseLimit = InTempMaxUpdatePhase;
+
+	if (!IsWithinPhaseLimit())
+	{
+		// No work needed if already past the phase limit
+		return EBuildResult::PhaseCompleted;
+	}
+	return BuildInternal();
+}
+
+void FTextBuilder::OnObjectsReplaced(const TMap<UObject*, UObject*>& InReplacementMap)
+{
+	if (UObject* const* TargetObject = InReplacementMap.Find(Component))
+	{
+		Component = Cast<UText3DComponent>(*TargetObject);
+	}
+
+	for (FExtensionElement& Extension : Extensions)
+	{
+		if (UObject* const* TargetObject = InReplacementMap.Find(Extension.Object))
+		{
+			Extension.Object = Cast<UText3DExtensionBase>(*TargetObject);
+		}
+	}
+}
+
+void FTextBuilder::AddReferencedObjects(FReferenceCollector& InCollector)
+{
+	InCollector.AddReferencedObject(Component);
+
+	for (FExtensionElement& Extension : Extensions)
+	{
+		InCollector.AddReferencedObject(Extension.Object);
+	}
+}
+
+EBuildResult FTextBuilder::BuildInternal()
+{
+	EBuildResult Result;
+
+	switch (Phase)
+	{
+	case EBuildPhase::NotStarted:
+		Result = PrepareComponent();
+		if (Extensions.IsEmpty())
+		{
+			Finish();
+			return EBuildResult::BuildCompleted;
+		}
+		break;
+
+	case EBuildPhase::PrepareExtensions:
+		Result = PrepareExtensions();
+		break;
+
+	case EBuildPhase::PreRendererUpdate:
+		Result = PreRendererUpdate();
+		break;
+
+	case EBuildPhase::RendererUpdate:
+		Result = RendererUpdate();
+		break;
+
+	case EBuildPhase::PostRendererUpdate:
+		Result = PostRendererUpdate();
+		break;
+
+	default:
+	case EBuildPhase::Finished:
+		Finish();
+		return EBuildResult::BuildCompleted;
+	}
+
+	// If work is complete, look to the next phase
+	if (Result == EBuildResult::PhaseCompleted)
+	{
+		return Next();
+	}
+	return Result;
+}
+
+bool FTextBuilder::IsWithinPhaseLimit() const
+{
+	return !TempPhaseLimit.IsSet() || Phase <= *TempPhaseLimit;
+}
+
+bool FTextBuilder::IsWithinBudget() const
+{
+	// If set, current time is within budget
+	const bool bWithinBudget = bSkipNextBudgetCheck || !BudgetEnd.IsSet() || FPlatformTime::Cycles64() <= *BudgetEnd;
+	bSkipNextBudgetCheck = false;
+	return bWithinBudget;
+}
+
+EBuildResult FTextBuilder::IterateExtensions(TFunctionRef<bool(FExtensionElement&)> InFunc)
+{
+	while (IsWithinBudget())
+	{
+		if (Extensions.IsValidIndex(ExtensionIndex))
+		{
+			FExtensionElement& Extension = Extensions[ExtensionIndex];
+			Extension.bActive = Extension.bActive && InFunc(Extension);
+			// Skip budget check if extension was not relevant (no-op)
+			bSkipNextBudgetCheck = !Extension.bActive;
+			++ExtensionIndex;
+		}
+		else
+		{
+			return EBuildResult::PhaseCompleted;
+		}
+	}
+	return EBuildResult::OutOfBudget;
+}
+
+EBuildResult FTextBuilder::IterateExtensions(TFunctionRef<bool(FExtensionElement&, const UE::Text3D::Renderer::FUpdateParameters&)> InFunc)
+{
+	const EText3DRendererFlags Flag = static_cast<EText3DRendererFlags>(CurrentFlag);
+	return IterateExtensions(
+		[Flag, &InFunc, this](FExtensionElement& InExtension)
+		{
+			if (!EnumHasAnyFlags(InExtension.FlagsLeft, Flag))
+			{
+				return true;
+			}
+			EnumRemoveFlags(InExtension.FlagsLeft, Flag);
+
+			const UE::Text3D::Renderer::FUpdateParameters Parameters
+				{
+					.UpdateFlags = Flags,
+					.CurrentFlag = Flag,
+					.bIsLastFlag = InExtension.FlagsLeft == EText3DRendererFlags::None
+				};
+			return InFunc(InExtension, Parameters);
+		});
+}
+
+EBuildResult FTextBuilder::PrepareComponent()
+{
+	TRACE_CPUPROFILER_EVENT_SCOPE(Text3D::PrepareComponent);
+
+	CurrentFlag = 1;
+	Extensions.Reset();
+
+	if (!Component->BeginBuild(&Flags))
+	{
+		return EBuildResult::BuildCompleted;
+	}
+
+	// Allocate extension elements
+	{
+		// Each component has a minimum of 7 extensions, 8 for a single layout effect.
+		TArray<UText3DExtensionBase*, TInlineAllocator<8>> DiscoveredExtensions;
+		DiscoveredExtensions.Reset();
+
+		Component->GatherExtensions(DiscoveredExtensions);
+		for (UText3DExtensionBase* Extension : DiscoveredExtensions)
+		{
+			FExtensionElement& ExtensionElement = Extensions.AddDefaulted_GetRef();
+			ExtensionElement.Object = Extension;
+		}
+	}
+
+	Extensions.StableSort(
+		[](const FExtensionElement& A, const FExtensionElement& B)->bool
+		{
+			return A.Object->GetUpdatePriority() < B.Object->GetUpdatePriority();
+		});
+
+	return EBuildResult::PhaseCompleted;
+}
+
+EBuildResult FTextBuilder::PrepareExtensions()
+{
+	TRACE_CPUPROFILER_EVENT_SCOPE(Text3D::PrepareExtensions);
+
+	return IterateExtensions(
+		[](FExtensionElement& InExtension)
+		{
+			InExtension.Object->PrepareBuild(InExtension.FlagsLeft);
+			return true;
+		});
+}
+
+EBuildResult FTextBuilder::PreRendererUpdate()
+{
+	TRACE_CPUPROFILER_EVENT_SCOPE(Text3D::PreRendererUpdate);
+	return IterateExtensions(
+		[](FExtensionElement& InExtension, const UE::Text3D::Renderer::FUpdateParameters& InParameters)
+		{
+			const EText3DExtensionResult Result = InExtension.Object->PreRendererUpdate(InParameters);
+			return Result == EText3DExtensionResult::Active;
+		});
+}
+
+EBuildResult FTextBuilder::RendererUpdate()
+{
+	if (!Component->TextRenderer)
+	{
+		return EBuildResult::PhaseCompleted;
+	}
+
+	TRACE_CPUPROFILER_EVENT_SCOPE(Text3D::RendererUpdate);
+
+	const EText3DRendererFlags Flag = static_cast<EText3DRendererFlags>(CurrentFlag);
+	if (!EnumHasAnyFlags(FlagsLeft, Flag))
+	{
+		return EBuildResult::PhaseCompleted;
+	}
+	EnumRemoveFlags(FlagsLeft, Flag);
+
+	const UE::Text3D::Renderer::FUpdateParameters Parameters
+		{
+			.UpdateFlags = Flags,
+			.CurrentFlag = Flag,
+			.bIsLastFlag = FlagsLeft == EText3DRendererFlags::None
+		};
+	Component->TextRenderer->Update(Parameters);
+	return EBuildResult::PhaseCompleted;
+}
+
+EBuildResult FTextBuilder::PostRendererUpdate()
+{
+	TRACE_CPUPROFILER_EVENT_SCOPE(Text3D::PostRendererUpdate);
+
+	const EBuildResult Result = IterateExtensions(
+		[](FExtensionElement& InExtension, const UE::Text3D::Renderer::FUpdateParameters& InParameters)
+		{
+			const EText3DExtensionResult Result = InExtension.Object->PostRendererUpdate(InParameters);
+			return Result == EText3DExtensionResult::Active;
+		});
+
+	if (Result == EBuildResult::PhaseCompleted)
+	{
+		const uint8 BuildFlags = static_cast<uint8>(Flags);
+		const uint8 AllFlags   = static_cast<uint8>(EText3DRendererFlags::All);
+
+		// Get the next flag to iterate
+		do
+		{
+			CurrentFlag <<= 1;
+
+			// If the end is reached, mark as phase completed without looping back.
+			if (CurrentFlag > AllFlags)
+			{
+				return EBuildResult::PhaseCompleted;
+			}
+		} 
+		while ((CurrentFlag & BuildFlags) == 0);
+
+		// Start loop again from PreRendererUpdate. This time with a different renderer flag.
+		Phase = static_cast<EBuildPhase>(static_cast<uint8>(EBuildPhase::PreRendererUpdate) - 1);
+	}
+	return Result;
+}
+
+EBuildResult FTextBuilder::Next()
+{
+	ExtensionIndex = 0;
+	FlagsLeft = Flags;
+	Phase = static_cast<EBuildPhase>(static_cast<uint8>(Phase) + 1);
+
+	// Reset FlagsLeft for each phase
+	for (FExtensionElement& Extension : Extensions)
+	{
+		Extension.FlagsLeft = Flags;
+	}
+
+	if (!IsWithinPhaseLimit())
+	{
+		return EBuildResult::PhaseCompleted;
+	}
+
+	if (IsWithinBudget())
+	{
+		return BuildInternal();
+	}
+
+	return EBuildResult::OutOfBudget;
+}
+
+void FTextBuilder::Finish()
+{
+	Component->FinishBuild(Flags);
+
+	Extensions.Reset();
+	ExtensionIndex = INDEX_NONE;
+	Phase = EBuildPhase::Finished;
+}
+
+} // UE::Text3D
diff --git a/Engine/Plugins/VirtualProduction/Text3D/Source/Text3D/Private/Text3DBuilder.h b/Engine/Plugins/VirtualProduction/Text3D/Source/Text3D/Private/Text3DBuilder.h
new file mode 100644
--- /dev/null
+++ b/Engine/Plugins/VirtualProduction/Text3D/Source/Text3D/Private/Text3DBuilder.h
@@ -0,0 +1,164 @@
+// Copyright Epic Games, Inc. All Rights Reserved.
+
+#pragma once
+
+#include "Text3DTypes.h"
+
+class UText3DComponent;
+class UText3DExtensionBase;
+
+namespace UE::Text3D
+{
+
+/** Enumeration of the different phases of the text building process */
+enum class EBuildPhase : uint8
+{
+	/** Build has not started yet */
+	NotStarted,
+	/** Prepares the Text3D extensions for building */
+	PrepareExtensions,
+	/** Calls PreRendererUpdate for the Text3D extensions */
+	PreRendererUpdate,
+	/** Runs the 'Update' for the Text3D renderer */
+	RendererUpdate,
+	/** Calls PostRendererUpdate for the Text3D extensions */
+	PostRendererUpdate,
+	/** Build finished */
+	Finished,
+};
+
+/** The build result for incremental builds  */
+enum class EBuildResult : uint8
+{
+	/** Build ran out of time budget */
+	OutOfBudget,
+	/** A build phase has completed */
+	PhaseCompleted,
+	/** Build completed */
+	BuildCompleted,
+};
+
+/** Manages the building of text for the Text3D Component */
+class FTextBuilder
+{
+	/** Holds the data needed for the Text3D extension during the build process */
+	struct FExtensionElement
+	{
+		/** Text3D extension being built */
+		TObjectPtr<UText3DExtensionBase> Object = nullptr;
+		/** Remaining flags left for the build phase. Resets every phase change. */
+		EText3DRendererFlags FlagsLeft = EText3DRendererFlags::None;
+		/** Whether the extension is active, or has completed and is not needed for the rest of the renderer updates */
+		bool bActive = true;
+	};
+
+public:
+	explicit FTextBuilder(UText3DComponent* InComponent);
+
+	/** 
+	 * Adds the given renderer flags to this builder. Should only be done for builders that have not started yet.
+	 * @return true if flags were added or false if nothing happened because the builder is already in progress.
+	 */
+	bool PrepareFlags(EText3DRendererFlags InFlags);
+
+	UText3DComponent* GetComponent() const
+	{
+		return Component;
+	}
+
+	EText3DRendererFlags GetFlags() const
+	{
+		return Flags;
+	}
+
+	/** Returns true if the builder is currently in progress (i.e. between not-started and finished) */
+	bool IsInProgress() const;
+
+	/** Returns true if the builder has finished building the component */
+	bool IsFinished() const;
+
+	/** Fully starts / resumes and finishes the build */
+	void BlockingBuild();
+
+	/** Incrementally builds the component every call with a budget end and update phase */
+	EBuildResult IncrementalBuild(TOptional<uint64> InBudgetEnd, TOptional<UE::Text3D::EBuildPhase> InTempMaxUpdatePhase = {});
+
+	/** Called by the owning text world subsystem when object replacement has taken place */
+	void OnObjectsReplaced(const TMap<UObject*, UObject*>& InReplacementMap);
+
+	/** Called by the owning text world subsystem to gather objects referenced by this builder */
+	void AddReferencedObjects(FReferenceCollector& InCollector);
+
+private:
+	/** Calls the appropriate function depending on the current phase */
+	EBuildResult BuildInternal();
+
+	/** Returns true if the current phase of the build is still less than or equal to the limit, or always true if no limit is set */
+	bool IsWithinPhaseLimit() const;
+
+	/** Returns true if the current cycles is less than the budgeted end cycle or if no budget has been set at all */
+	bool IsWithinBudget() const;
+
+	EBuildResult IterateExtensions(TFunctionRef<bool(FExtensionElement&)> InFunc);
+	EBuildResult IterateExtensions(TFunctionRef<bool(FExtensionElement&, const UE::Text3D::Renderer::FUpdateParameters&)> InFunc);
+
+	/** Phase 1 - calls BeginBuild on the component and gathers the component's extensions */
+	EBuildResult PrepareComponent();
+
+	/** Phase 2- calls PrepareBuild on the extensions */
+	EBuildResult PrepareExtensions();
+
+	/** Phase 3 - calls PreRendererUpdate on the extensions for a single renderer flag */
+	EBuildResult PreRendererUpdate();
+
+	/** Phase 4 - calls Update on the component's renderer for a single renderer flag */
+	EBuildResult RendererUpdate();
+
+	/**
+	 * Phase 5 - calls PostRendererUpdate on the extensions for a single renderer flag 
+	 * Once this phase finishes, the next available renderer flag is gotten and if available, build loops back to Phase 3.
+	 * If there's no available renderer flag left the build finishes.
+	 */
+	EBuildResult PostRendererUpdate();
+
+	/** Called to move onto the next phase */
+	EBuildResult Next();
+
+	/** Called when the build has completed */
+	void Finish();
+
+	/** Component to build */
+	TObjectPtr<UText3DComponent> Component;
+
+	/** Extensions being built */
+	TArray<FExtensionElement> Extensions;
+
+	/** If set, the cycle that the build should stop when past it  */
+	TOptional<uint64> BudgetEnd;
+
+	/** Current extension being updated. Relevant only to certain phases */
+	int32 ExtensionIndex = INDEX_NONE;
+
+	/** Renderer flags for the built */
+	EText3DRendererFlags Flags = EText3DRendererFlags::None;
+
+	/** Remaining renderer flags for the current phase. Reset on every phase change */
+	EText3DRendererFlags FlagsLeft = EText3DRendererFlags::None;
+
+	/** Current renderer flag. Relevant only to the renderer phases. Does not reset on phase change. */
+	uint8 CurrentFlag = 1;
+
+	/** Current build phase */
+	EBuildPhase Phase = EBuildPhase::NotStarted;
+
+	/**
+	 * Temporary max build phase that is allowed by the incremental update.
+	 * Set to a new value every incremental update call.
+	 */
+	TOptional<EBuildPhase> TempPhaseLimit;
+
+	/** Whether budget check should be skipped to avoid querying cycles when unnecessary. Resets every IsWithinBudget() call */
+	mutable bool bSkipNextBudgetCheck = false;
+};
+
+} // UE::Text3D
diff --git a/Engine/Plugins/VirtualProduction/Text3D/Source/Text3D/Private/Text3DComponent.cpp b/Engine/Plugins/VirtualProduction/Text3D/Source/Text3D/Private/Text3DComponent.cpp
--- a/Engine/Plugins/VirtualProduction/Text3D/Source/Text3D/Private/Text3DComponent.cpp
+++ b/Engine/Plugins/VirtualProduction/Text3D/Source/Text3D/Private/Text3DComponent.cpp
@@ -11,8 +11,8 @@
 #include "Extensions/Text3DDefaultMaterialExtension.h"
 #include "Extensions/Text3DDefaultRenderingExtension.h"
 #include "Extensions/Text3DDefaultStyleExtension.h"
-#include "Extensions/Text3DLayoutEffectBase.h"
 #include "Extensions/Text3DDefaultTokenExtension.h"
+#include "Extensions/Text3DLayoutEffectBase.h"
 #include "Fonts/CompositeFont.h"
 #include "Internationalization/Regex.h"
 #include "Logs/Text3DLogs.h"
@@ -23,6 +23,7 @@
 #include "Renderers/StaticMeshes/Text3DStaticMeshesRenderer.h"
 #include "Renderers/Text3DRendererBase.h"
 #include "Settings/Text3DProjectSettings.h"
+#include "Subsystems/Text3DBuildSystem.h"
 #include "Text3DComponentVersion.h"
 #include "UObject/ConstructorHelpers.h"
 #include "UObject/UObjectGlobals.h"
@@ -36,7 +37,6 @@ UText3DComponent::FOnResolveFontByName UText3DComponent::OnResolveFontByNameDele
 FCustomVersionRegistration GRegisterText3DComponentVersion(FText3DComponentVersion::GUID, static_cast<int32>(FText3DComponentVersion::LatestVersion), TEXT("Text3DComponentVersion"));
 
 UText3DComponent::UText3DComponent()
-	: bIsUpdatingText(false)
 {
 	const UText3DProjectSettings* Settings = UText3DProjectSettings::Get();
 
@@ -66,43 +66,42 @@ UText3DComponent::UText3DComponent()
 void UText3DComponent::OnComponentCreated()
 {
 	Super::OnComponentCreated();
-	RequestUpdate(EText3DRendererFlags::All, /** Immediate */ true);
+	RequestUpdate(EText3DRendererFlags::All);
 }
 
 void UText3DComponent::OnComponentDestroyed(bool bInDestroyingHierarchy)
 {
 	Super::OnComponentDestroyed(bInDestroyingHierarchy);
 
-	if (TextRenderer)
+	if (UText3DBuildSystem* BuildSystem = UText3DBuildSystem::Get(this))
 	{
-		TextRenderer->Destroy();
+		BuildSystem->RemoveTextBuild(this);
 	}
 
 	TransformUpdated.RemoveAll(this);
-
-	FTSTicker::RemoveTicker(TextUpdateHandle);
-	TextUpdateHandle.Reset();
 }
 
 void UText3DComponent::PostLoad()
 {
 	Super::PostLoad();
 	RefreshTypeface();
-	const bool bCompiledFromBlueprint = IsCompiledFromBlueprint();
-	RequestUpdate(EText3DRendererFlags::All, /** Immediate */!bCompiledFromBlueprint);
+	RequestUpdate(EText3DRendererFlags::All);
 }
 
 void UText3DComponent::RequestUpdate(EText3DRendererFlags InFlags, bool bInImmediate)
 {
 	if (InFlags != EText3DRendererFlags::None && !IsTemplate(RF_ClassDefaultObject))
 	{
-		if (bInImmediate)
+		if (IsRegistered())
 		{
-			Rebuild(InFlags);
+			if (UText3DBuildSystem* const BuildSystem = UText3DBuildSystem::Get(this))
+			{
+				BuildSystem->RequestTextBuild(this, InFlags, bInImmediate);
+			}
 		}
 		else
 		{
-			QueueRebuild(InFlags);
+			EnumAddFlags(PreRegisterFlags, InFlags);
 		}
 	}
 }
@@ -272,20 +271,18 @@ bool UText3DComponent::IsCompiledFromBlueprint() const
 void UText3DComponent::OnRegister()
 {
 	Super::OnRegister();
-	const bool bCompiledFromBlueprint = IsCompiledFromBlueprint();
-	RequestUpdate(UpdateFlags, /** Immediate */!bCompiledFromBlueprint);
+
+	RequestUpdate(PreRegisterFlags);
+	PreRegisterFlags = EText3DRendererFlags::None;
 }
 
 void UText3DComponent::OnUnregister()
 {
 	Super::OnUnregister();
 
-	if (IsBeingDestroyed())
+	if (UText3DBuildSystem* BuildSystem = UText3DBuildSystem::Get(this))
 	{
-		if (TextRenderer)
-		{
-			TextRenderer->Destroy();
-		}
+		BuildSystem->RemoveTextBuild(this);
 	}
 }
 
@@ -373,12 +370,6 @@ void UText3DComponent::PostDuplicate(bool bDuplicateForPIE)
 void UText3DComponent::PostCDOCompiled(const FPostCDOCompiledContext& InContext)
 {
 	Super::PostCDOCompiled(InContext);
-
-	// Cleanup old components
-	if (TextRenderer)
-	{
-		TextRenderer->Destroy();
-	}
 }
 
 void UText3DComponent::PostEditChangeProperty(FPropertyChangedEvent& PropertyChangedEvent)
@@ -1007,154 +998,6 @@ void UText3DComponent::SetTextRendererClass(const TSubclassOf<UText3DRendererBas
 	OnTextRendererClassChanged();
 }
 
-void UText3DComponent::QueueRebuild(EText3DRendererFlags InFlags)
-{
-	// Add the given flags to the current queued flags in case a text update is already queued
-	EnumAddFlags(UpdateFlags, InFlags);
-
-	if (!TextUpdateHandle.IsValid())
-	{
-		TextUpdateHandle = FTSTicker::GetCoreTicker().AddTicker(
-			FTickerDelegate::CreateWeakLambda(
-				this,
-				[this](float)->bool
-				{
-					if (UE::IsSavingPackage(this))
-					{
-						// Postpone to next tick
-						return true;
-					}
-
-					// Reset handle so that while building text, new rebuild requests can be issued
-					// New rebuild requests should only happen when, while building text, existing 'built data' becomes outdated.
-					TextUpdateHandle.Reset();
-
-					if (!IsRegistered())
-					{
-						// Handled in OnRegister()
-						return false;
-					}
-
-					const EText3DRendererFlags CurrentUpdateFlags = UpdateFlags;
-					UpdateFlags = EText3DRendererFlags::None;
-
-					Rebuild(CurrentUpdateFlags);
-
-					// Executes only once
-					return false;
-				}
-			)
-		);
-	}
-}
-
-void UText3DComponent::Rebuild(EText3DRendererFlags InUpdateFlags)
-{
-	if (InUpdateFlags == EText3DRendererFlags::None)
-	{
-		return;
-	}
-
-#if !WITH_FREETYPE
-	UE_LOG(LogText3D, Warning, TEXT("FreeType is not available, cannot proceed without it"));
-	return;
-#endif
-
-	if (!TextRenderer || TextRenderer->GetClass() != TextRendererClass)
-	{
-		OnTextRendererClassChanged();
-
-		if (!TextRenderer)
-		{
-			UE_LOG(LogText3D, Warning, TEXT("Text3D renderer is not valid, cannot proceed"));
-			return;
-		}
-	}
-
-	if (!Font)
-	{
-		if (const UText3DProjectSettings* Text3DSettings = UText3DProjectSettings::Get())
-		{
-			Font = Text3DSettings->GetFallbackFont();
-		}
-
-		if (!Font)
-		{
-			UE_LOG(LogText3D, Warning, TEXT("Font object is not valid, Fallback font is not defined in project settings, cannot proceed"));
-			return;
-		}
-	}
-
-	if (!IsRegistered())
-	{
-		if (!EnumHasAllFlags(UpdateFlags, InUpdateFlags))
-		{
-			UE_LOG(LogText3D, Verbose, TEXT("[%s] Rebuild called with update flags (%u) while component is not registered yet. This new Rebuild will be deferred to the next tick."), *GetReadableName(), InUpdateFlags);
-			QueueRebuild(InUpdateFlags);
-		}
-
-		return;
-	}
-
-	if (UE::IsSavingPackage(this))
-	{
-		if (!EnumHasAllFlags(UpdateFlags, InUpdateFlags))
-		{
-			UE_LOG(LogText3D, Verbose, TEXT("[%s] Rebuild called with update flags (%u) while package is being saved. This new Rebuild will be deferred to the next tick."), *GetReadableName(), InUpdateFlags);
-			QueueRebuild(InUpdateFlags);
-		}
-
-		return;
-	}
-
-	if (bIsUpdatingText)
-	{
-		if (!EnumHasAllFlags(UpdateFlags, InUpdateFlags))
-		{
-			UE_LOG(LogText3D, Log, TEXT("[%s] Rebuild called with update flags (%u) while building text. This new Rebuild will be deferred to the next tick."), *GetReadableName(), InUpdateFlags);
-			QueueRebuild(InUpdateFlags);
-		}
-
-		return;
-	}
-
-	EnumRemoveFlags(UpdateFlags, InUpdateFlags);
-
-	bIsUpdatingText = true;
-
-	ON_SCOPE_EXIT
-	{
-		bIsUpdatingText = false;
-	};
-
-	if (!TextRenderer->IsInitialized())
-	{
-		TextRenderer->Create();
-		InUpdateFlags = EText3DRendererFlags::All;
-	}
-
-	TextPreUpdateDelegate.Broadcast(this, InUpdateFlags);
-
-	TRACE_CPUPROFILER_EVENT_SCOPE(UText3DComponent::Rebuild);
-
-	TextRenderer->Update(InUpdateFlags);
-
-	if (EnumHasAnyFlags(InUpdateFlags,EText3DRendererFlags::Geometry))
-	{
-		UpdateStatistics();
-
-		TextGeneratedNativeDelegate.Broadcast();
-		TextGeneratedDelegate.Broadcast();
-	}
-
-	if (EnumHasAnyFlags(InUpdateFlags,EText3DRendererFlags::Visibility))
-	{
-		MarkRenderStateDirty();
-	}
-
-	TextPostUpdateDelegate.Broadcast(this, InUpdateFlags);
-}
-
 void UText3DComponent::OnTextRendererClassChanged()
 {
 	if (IsTemplate(RF_ClassDefaultObject))
@@ -1210,14 +1053,13 @@ void UText3DComponent::OnTextRendererClassChanged()
 
 	if (!bIsUpdatingText)
 	{
-		RequestUpdate(EText3DRendererFlags::All, /** Immediate */true);
+		RequestUpdate(EText3DRendererFlags::All);
 	}
 }
 
 void UText3DComponent::OnTextChanged()
 {
-	constexpr bool bImmediate = true;
-	RequestUpdate(EText3DRendererFlags::All, bImmediate);
+	RequestUpdate(EText3DRendererFlags::All);
 }
 
 void UText3DComponent::OnFontPropertiesChanged()
@@ -1260,6 +1102,134 @@ TArray<FName> UText3DComponent::GetTypefaceNames() const
 	return TypefaceNames;
 }
 
+bool UText3DComponent::GatherExtensions(TAdderReserverRef<UText3DExtensionBase*> OutExtensions) const
+{
+	// Let's make sure text extension subobjects are valid and usable !
+	if (!TokenExtension)
+	{
+		UE_LOG(LogText3D, Warning, TEXT("Invalid Text3D Token extension, cannot proceed !"))
+		return false;
+	}
+
+	if (!StyleExtension)
+	{
+		UE_LOG(LogText3D, Warning, TEXT("Invalid Text3D Style extension, cannot proceed !"))
+		return false;
+	}
+
+	if (!CharacterExtension)
+	{
+		UE_LOG(LogText3D, Warning, TEXT("Invalid Text3D Character extension, cannot proceed !"))
+		return false;
+	}
+
+	if (!GeometryExtension)
+	{
+		UE_LOG(LogText3D, Warning, TEXT("Invalid Text3D Geometry extension, cannot proceed !"))
+		return false;
+	}
+
+	if (!LayoutExtension)
+	{
+		UE_LOG(LogText3D, Warning, TEXT("Invalid Text3D Layout extension, cannot proceed !"))
+		return false;
+	}
+
+	if (!MaterialExtension)
+	{
+		UE_LOG(LogText3D, Warning, TEXT("Invalid Text3D Material extension, cannot proceed !"))
+		return false;
+	}
+
+	if (!RenderingExtension)
+	{
+		UE_LOG(LogText3D, Warning, TEXT("Invalid Text3D Rendering extension, cannot proceed !"))
+		return false;
+	}
+
+	OutExtensions.Reserve(7 + LayoutEffects.Num());
+	OutExtensions.Add(TokenExtension);
+	OutExtensions.Add(StyleExtension);
+	OutExtensions.Add(CharacterExtension);
+	OutExtensions.Add(GeometryExtension);
+	OutExtensions.Add(LayoutExtension);
+	OutExtensions.Add(MaterialExtension);
+	OutExtensions.Add(RenderingExtension);
+
+	// Add optional extensions for layout effects
+	for (UText3DLayoutEffectBase* LayoutEffect : LayoutEffects)
+	{
+		if (LayoutEffect)
+		{
+			OutExtensions.Add(LayoutEffect);
+		}
+	}
+
+	return true;
+}
+
+bool UText3DComponent::BeginBuild(TNotNull<EText3DRendererFlags*> InOutFlags)
+{
+	bIsUpdatingText = true;
+	TextPreUpdateDelegate.Broadcast(this, *InOutFlags);
+
+	// If only updating materials, no need to prepare TextRenderer, or Font because the material extension is mostly self contained.
+	// TODO: asking extensions whether TextRenderer/etc is needed when building with specific flags?
+	if (*InOutFlags == EText3DRendererFlags::Material)
+	{
+		return true;
+	}
+
+	if (!TextRenderer || TextRenderer->GetClass() != TextRendererClass)
+	{
+		OnTextRendererClassChanged();
+		if (!TextRenderer)
+		{
+			UE_LOG(LogText3D, Warning, TEXT("Text3D renderer is not valid, cannot proceed"));
+			return false;
+		}
+	}
+
+	if (!Font)
+	{
+		if (const UText3DProjectSettings* Text3DSettings = UText3DProjectSettings::Get())
+		{
+			Font = Text3DSettings->GetFallbackFont();
+		}
+
+		if (!Font)
+		{
+			UE_LOG(LogText3D, Warning, TEXT("Font object is not valid, Fallback font is not defined in project settings, cannot proceed"));
+			return false;
+		}
+	}
+
+	if (!TextRenderer->IsInitialized())
+	{
+		TextRenderer->Create();
+		*InOutFlags = EText3DRendererFlags::All;
+	}
+	return true;
+}
+
+void UText3DComponent::FinishBuild(EText3DRendererFlags InFlags)
+{
+	if (EnumHasAnyFlags(InFlags,EText3DRendererFlags::Geometry))
+	{
+		UpdateStatistics();
+		TextGeneratedNativeDelegate.Broadcast();
+		TextGeneratedDelegate.Broadcast();
+	}
+
+	if (EnumHasAnyFlags(InFlags,EText3DRendererFlags::Visibility))
+	{
+		MarkRenderStateDirty();
+	}
+
+	TextPostUpdateDelegate.Broadcast(this, InFlags);
+	bIsUpdatingText = false;
+}
+
 void UText3DComponent::OnVisibilityChanged()
 {
 	Super::OnVisibilityChanged();
diff --git a/Engine/Plugins/VirtualProduction/Text3D/Source/Text3D/Private/Utilities/Text3DUtilities.cpp b/Engine/Plugins/VirtualProduction/Text3D/Source/Text3D/Private/Utilities/Text3DUtilities.cpp
--- a/Engine/Plugins/VirtualProduction/Text3D/Source/Text3D/Private/Utilities/Text3DUtilities.cpp
+++ b/Engine/Plugins/VirtualProduction/Text3D/Source/Text3D/Private/Utilities/Text3DUtilities.cpp
@@ -2,10 +2,12 @@
 
 #include "Utilities/Text3DUtilities.h"
 
+#include "Engine/Engine.h"
 #include "Engine/Font.h"
 #include "Engine/FontFace.h"
 #include "Fonts/CompositeFont.h"
 #include "Misc/Paths.h"
+#include "Subsystems/Text3DBuildSystem.h"
 #include "Text3DComponent.h"
 #include "Text3DModule.h"
 
@@ -33,6 +35,19 @@ THIRD_PARTY_INCLUDES_END
 
 #define LOCTEXT_NAMESPACE "Text3DUtilities"
 
+IText3DBuildSystemInterface* UE::Text3D::Utilities::FindBuildSystem(UObject* InWorldContextObject)
+{
+	if (!GEngine)
+	{
+		return nullptr;
+	}
+	if (const UWorld* const World = GEngine->GetWorldFromContextObject(InWorldContextObject, EGetWorldErrorMode::LogAndReturnNull))
+	{
+		return World->GetSubsystem<UText3DBuildSystem>();
+	}
+	return nullptr;
+}
+
 bool UE::Text3D::Utilities::Font::GetSanitizeFontName(const UFont* InFont, FString& OutFontName)
 {
 	if (GetFontName(InFont, OutFontName))
diff --git a/Engine/Plugins/VirtualProduction/Text3D/Source/Text3D/Public/Extensions/Text3DDefaultGeometryExtension.h b/Engine/Plugins/VirtualProduction/Text3D/Source/Text3D/Public/Extensions/Text3DDefaultGeometryExtension.h
--- a/Engine/Plugins/VirtualProduction/Text3D/Source/Text3D/Public/Extensions/Text3DDefaultGeometryExtension.h
+++ b/Engine/Plugins/VirtualProduction/Text3D/Source/Text3D/Public/Extensions/Text3DDefaultGeometryExtension.h
@@ -125,6 +125,7 @@ protected:
 	//~ End UText3DGeometryExtensionBase
 
 	//~ Begin UText3DExtensionBase
+	virtual void PrepareBuild(EText3DRendererFlags InUpdateFlags) override;
 	virtual EText3DExtensionResult PreRendererUpdate(const UE::Text3D::Renderer::FUpdateParameters& InParameters) override;
 	virtual EText3DExtensionResult PostRendererUpdate(const UE::Text3D::Renderer::FUpdateParameters& InParameters) override;
 	//~ End UText3DExtensionBase
diff --git a/Engine/Plugins/VirtualProduction/Text3D/Source/Text3D/Public/Extensions/Text3DDefaultLayoutExtension.h b/Engine/Plugins/VirtualProduction/Text3D/Source/Text3D/Public/Extensions/Text3DDefaultLayoutExtension.h
--- a/Engine/Plugins/VirtualProduction/Text3D/Source/Text3D/Public/Extensions/Text3DDefaultLayoutExtension.h
+++ b/Engine/Plugins/VirtualProduction/Text3D/Source/Text3D/Public/Extensions/Text3DDefaultLayoutExtension.h
@@ -173,6 +173,7 @@ protected:
 	//~ End UText3DLayoutExtensionBase
 
 	//~ Begin UText3DExtensionBase
+	virtual void PrepareBuild(EText3DRendererFlags InUpdateFlags) override;
 	virtual EText3DExtensionResult PreRendererUpdate(const UE::Text3D::Renderer::FUpdateParameters& InParameters) override;
 	virtual EText3DExtensionResult PostRendererUpdate(const UE::Text3D::Renderer::FUpdateParameters& InParameters) override;
 	//~ End UText3DExtensionBase
diff --git a/Engine/Plugins/VirtualProduction/Text3D/Source/Text3D/Public/Extensions/Text3DDefaultStyleExtension.h b/Engine/Plugins/VirtualProduction/Text3D/Source/Text3D/Public/Extensions/Text3DDefaultStyleExtension.h
--- a/Engine/Plugins/VirtualProduction/Text3D/Source/Text3D/Public/Extensions/Text3DDefaultStyleExtension.h
+++ b/Engine/Plugins/VirtualProduction/Text3D/Source/Text3D/Public/Extensions/Text3DDefaultStyleExtension.h
@@ -41,8 +41,9 @@ protected:
     virtual TSharedPtr<FSlateStyleSet> GetCustomStyles() const override;
 	virtual UText3DStyleBase* GetStyle(FName InStyle) const override;
     //~ End UText3DStyleExtensionBase
-    
+
     //~ Begin UText3DExtensionBase
+	virtual void PrepareBuild(EText3DRendererFlags InUpdateFlags) override;
     virtual EText3DExtensionResult PreRendererUpdate(const UE::Text3D::Renderer::FUpdateParameters& InParameters) override;
     virtual EText3DExtensionResult PostRendererUpdate(const UE::Text3D::Renderer::FUpdateParameters& InParameters) override;
     //~ End UText3DExtensionBase
diff --git a/Engine/Plugins/VirtualProduction/Text3D/Source/Text3D/Public/Extensions/Text3DDefaultTokenExtension.h b/Engine/Plugins/VirtualProduction/Text3D/Source/Text3D/Public/Extensions/Text3DDefaultTokenExtension.h
--- a/Engine/Plugins/VirtualProduction/Text3D/Source/Text3D/Public/Extensions/Text3DDefaultTokenExtension.h
+++ b/Engine/Plugins/VirtualProduction/Text3D/Source/Text3D/Public/Extensions/Text3DDefaultTokenExtension.h
@@ -26,6 +26,7 @@ protected:
 	//~ End UText3DTokenExtensionBase
 
 	//~ Begin UText3DExtensionBase
+	virtual void PrepareBuild(EText3DRendererFlags InUpdateFlags) override;
 	virtual EText3DExtensionResult PreRendererUpdate(const UE::Text3D::Renderer::FUpdateParameters& InParameters) override;
 	virtual EText3DExtensionResult PostRendererUpdate(const UE::Text3D::Renderer::FUpdateParameters& InParameters) override;
 	//~ End UText3DExtensionBase
diff --git a/Engine/Plugins/VirtualProduction/Text3D/Source/Text3D/Public/Extensions/Text3DExtensionBase.h b/Engine/Plugins/VirtualProduction/Text3D/Source/Text3D/Public/Extensions/Text3DExtensionBase.h
--- a/Engine/Plugins/VirtualProduction/Text3D/Source/Text3D/Public/Extensions/Text3DExtensionBase.h
+++ b/Engine/Plugins/VirtualProduction/Text3D/Source/Text3D/Public/Extensions/Text3DExtensionBase.h
@@ -36,6 +36,11 @@ public:
 		: UpdatePriority(InPriority)
 	{}
 
+	/** Prepare the extensions before a build occurs */
+	virtual void PrepareBuild(EText3DRendererFlags InUpdateFlags)
+	{
+	}
+
 	/** Perform an update of the extension behavior before renderer execute based on condition above */
 	virtual EText3DExtensionResult PreRendererUpdate(const UE::Text3D::Renderer::FUpdateParameters& InParameters)
 	{
diff --git a/Engine/Plugins/VirtualProduction/Text3D/Source/Text3D/Public/IText3DBuildSystem.h b/Engine/Plugins/VirtualProduction/Text3D/Source/Text3D/Public/IText3DBuildSystem.h
new file mode 100644
--- /dev/null
+++ b/Engine/Plugins/VirtualProduction/Text3D/Source/Text3D/Public/IText3DBuildSystem.h
@@ -0,0 +1,21 @@
+// Copyright Epic Games, Inc. All Rights Reserved.
+
+#pragma once
+
+#include "UObject/Interface.h"
+#include "IText3DBuildSystem.generated.h"
+
+UINTERFACE(MinimalAPI, meta=(CannotImplementInterfaceInBlueprint))
+class UText3DBuildSystemInterface : public UInterface
+{
+	GENERATED_BODY()
+};
+
+class IText3DBuildSystemInterface
+{
+	GENERATED_BODY()
+
+public:
+	/** Returns true if there are builders yet to complete */
+	virtual bool IsPendingBuildCompletion() const = 0;
+};
diff --git a/Engine/Plugins/VirtualProduction/Text3D/Source/Text3D/Public/Renderers/Text3DRendererBase.h b/Engine/Plugins/VirtualProduction/Text3D/Source/Text3D/Public/Renderers/Text3DRendererBase.h
--- a/Engine/Plugins/VirtualProduction/Text3D/Source/Text3D/Public/Renderers/Text3DRendererBase.h
+++ b/Engine/Plugins/VirtualProduction/Text3D/Source/Text3D/Public/Renderers/Text3DRendererBase.h
@@ -23,7 +23,7 @@ public:
 	void Create();
 
 	/** Update rendering state */
-	void Update(EText3DRendererFlags InFlags);
+	void Update(const UE::Text3D::Renderer::FUpdateParameters& InParameters);
 
 	/** Clears the active rendering state */
 	void Clear();
diff --git a/Engine/Plugins/VirtualProduction/Text3D/Source/Text3D/Public/Text3DComponent.h b/Engine/Plugins/VirtualProduction/Text3D/Source/Text3D/Public/Text3DComponent.h
--- a/Engine/Plugins/VirtualProduction/Text3D/Source/Text3D/Public/Text3DComponent.h
+++ b/Engine/Plugins/VirtualProduction/Text3D/Source/Text3D/Public/Text3DComponent.h
@@ -3,22 +3,28 @@
 #pragma once
 
 #include "Components/SceneComponent.h"
-#include "Containers/Ticker.h"
 #include "Text3DTypes.h"
+#include "Containers/AdderRef.h"
 #include "Text3DComponent.generated.h"
 
+namespace UE::Text3D
+{
+	class FTextBuilder;
+}
+
 class UFont;
 class UMaterialInterface;
 class UText3DCharacterBase;
 class UText3DCharacterExtensionBase;
+class UText3DExtensionBase;
 class UText3DGeometryExtensionBase;
 class UText3DLayoutEffectBase;
 class UText3DLayoutExtensionBase;
 class UText3DMaterialExtensionBase;
-class UText3DStyleExtensionBase;
-class UText3DTokenExtensionBase;
 class UText3DRendererBase;
 class UText3DRenderingExtensionBase;
+class UText3DStyleExtensionBase;
+class UText3DTokenExtensionBase;
 struct FTypefaceEntry;
 
 UCLASS(MinimalAPI
@@ -409,6 +415,15 @@ public:
 		});
 	}
 
+	/** Retrieves all the extensions for this component (not sorted) */
+	bool GatherExtensions(TAdderReserverRef<UText3DExtensionBase*> OutExtensions) const;
+
+	/** Prepares the component for a build */
+	bool BeginBuild(TNotNull<EText3DRendererFlags*> InOutFlags);
+
+	/** Called when the text has finished building */
+	void FinishBuild(EText3DRendererFlags InFlags);
+
 protected:
 	//~ Begin USceneComponent
 	/** Intercept and propagate a change on this component to all children. */
@@ -539,27 +554,18 @@ private:
 	/** Called before text is updated */
 	FTextUpdated TextPreUpdateDelegate;
 
+	/** Flags from requests prior to component registration */
+	EText3DRendererFlags PreRegisterFlags = EText3DRendererFlags::None;
+
 	/** Called after text is updated */
 	FTextUpdated TextPostUpdateDelegate;
 
 	/** Flagged as true when text is being updated */
-	std::atomic<bool> bIsUpdatingText;
-
-	/** Used to determine and selectively perform the the type of rebuild requested. */
-	EText3DRendererFlags UpdateFlags = EText3DRendererFlags::None;
+	bool bIsUpdatingText = false;
 
 	/** Text statistics cached since last text generation */
 	FText3DStatistics Statistics;
 
-	/** Planned update requested */
-	FTSTicker::FDelegateHandle TextUpdateHandle;
-
-	/** Queues text rebuild for the next tick */
-	void QueueRebuild(EText3DRendererFlags InFlags);
-
-	/** Partially or fully rebuilds text depending on the given update flags */
-	void Rebuild(EText3DRendererFlags InUpdateFlags);
-
 	bool IsTypefaceAvailable(FName InTypeface) const;
 	TArrayView<const FTypefaceEntry> GetAvailableTypefaces() const;
 	void RefreshTypeface();
@@ -661,4 +667,6 @@ private:
 	UPROPERTY()
 	bool bCastShadow_DEPRECATED = true;
 #endif
+
+	friend class UE::Text3D::FTextBuilder;
 };
diff --git a/Engine/Plugins/VirtualProduction/Text3D/Source/Text3D/Public/Utilities/Text3DUtilities.h b/Engine/Plugins/VirtualProduction/Text3D/Source/Text3D/Public/Utilities/Text3DUtilities.h
--- a/Engine/Plugins/VirtualProduction/Text3D/Source/Text3D/Public/Utilities/Text3DUtilities.h
+++ b/Engine/Plugins/VirtualProduction/Text3D/Source/Text3D/Public/Utilities/Text3DUtilities.h
@@ -6,8 +6,9 @@
 #include "Text3DTypes.h"
 
 class AStaticMeshActor;
-class UFontFace;
+class IText3DBuildSystemInterface;
 class UFont;
+class UFontFace;
 class UText3DComponent;
 struct FText3DTypeFaceMetrics;
 struct FTypefaceEntry;
@@ -15,6 +16,9 @@ struct FTypefaceEntry;
 /** Utilities to manipulate text font */
 namespace UE::Text3D::Utilities
 {
+	/** Finds the text build system for the given world context object */
+	TEXT3D_API IText3DBuildSystemInterface* FindBuildSystem(UObject* InWorldContextObject);
+
 	namespace Font
 	{
 		/** Get proper font name without invalid chars */
