Text3D: fix static analysis warning of iterating empty array

#Synced-CL 51290016

diff --git a/Engine/Plugins/VirtualProduction/Text3D/Source/Text3D/Private/Text3DBuilder.cpp b/Engine/Plugins/VirtualProduction/Text3D/Source/Text3D/Private/Text3DBuilder.cpp
--- a/Engine/Plugins/VirtualProduction/Text3D/Source/Text3D/Private/Text3DBuilder.cpp
+++ b/Engine/Plugins/VirtualProduction/Text3D/Source/Text3D/Private/Text3DBuilder.cpp
@@ -215,6 +215,12 @@ EBuildResult FTextBuilder::PrepareComponent()
 		DiscoveredExtensions.Reset();
 
 		Component->GatherExtensions(DiscoveredExtensions);
+		if (DiscoveredExtensions.IsEmpty())
+		{
+			return EBuildResult::BuildCompleted;
+		}
+
+		Extensions.Reserve(DiscoveredExtensions.Num());
 		for (UText3DExtensionBase* Extension : DiscoveredExtensions)
 		{
 			FExtensionElement& ExtensionElement = Extensions.AddDefaulted_GetRef();
