Text3D: fix issue with changing properties not updating the text.
+ Added verbose logs for debugging.

#Synced-CL 51290016

diff --git a/Engine/Plugins/VirtualProduction/Text3D/Source/Text3D/Private/Subsystems/Text3DBuildSystem.cpp b/Engine/Plugins/VirtualProduction/Text3D/Source/Text3D/Private/Subsystems/Text3DBuildSystem.cpp
--- a/Engine/Plugins/VirtualProduction/Text3D/Source/Text3D/Private/Subsystems/Text3DBuildSystem.cpp
+++ b/Engine/Plugins/VirtualProduction/Text3D/Source/Text3D/Private/Subsystems/Text3DBuildSystem.cpp
@@ -45,25 +45,39 @@ void UText3DBuildSystem::RequestTextBuild(TNotNull<UText3DComponent*> InComponen
 		return;
 	}
 
-	auto CanUseBuilder = 
-		[InFlags, bInBlockingBuild, InComponent](const UE::Text3D::FTextBuilder& InBuilder)->bool
+	UE_LOG(LogText3D, Verbose, TEXT("[%llu] Requesting text build for component '%s'. Flags: %d. Blocking? %d")
+		, GFrameCounter
+		, *InComponent->GetFullName()
+		, static_cast<int32>(InFlags)
+		, bInBlockingBuild);
+
+	// allow if build is incremental. If blocking ensure the flags are exactly equal. If there's another flag, it can't be used.
+	auto IsBuilderFlagsCompatible = 
+		[InFlags, bInBlockingBuild](const UE::Text3D::FTextBuilder& InBuilder)->bool
 		{
-			// matching component
-			return InBuilder.GetComponent() == InComponent
-				// not currently building
-				&& !InBuilder.IsInProgress()
-				// allow if build is incremental. If blocking ensure the flags are exactly equal. If there's another flag, it can't be used.
-				&& (!bInBlockingBuild || InBuilder.GetFlags() == InFlags);
+			return !bInBlockingBuild || InBuilder.GetFlags() == InFlags;
 		};
 
-	// Check in active builders. There could still be builders that haven't yet started where we could queue in our request.
+	// Cancel builders in progress if flags overlap. There could still be builders that haven't yet started where we could queue in our request.
 	UE::Text3D::FTextBuilder* Builder = nullptr;
 	for (UE::Text3D::FTextBuilder& CandidateBuilder : TextBuilders)
 	{
-		if (CanUseBuilder(CandidateBuilder) && EnumHasAllFlags(CandidateBuilder.GetFlags(), InFlags))
+		if (CandidateBuilder.GetComponent() == InComponent && EnumHasAnyFlags(CandidateBuilder.GetFlags(), InFlags))
 		{
-			Builder = &CandidateBuilder;
-			break;
+			// Cancel builder if already in progress. A newer update is available.
+			if (CandidateBuilder.IsInProgress())
+			{
+				UE_LOG(LogText3D, Verbose, TEXT("\tCanceling active builder '%s'. A newer build request was made."), *CandidateBuilder.GetDebugString());
+				CandidateBuilder.CancelBuild();
+				continue;
+			}
+
+			if (IsBuilderFlagsCompatible(CandidateBuilder) && EnumHasAllFlags(CandidateBuilder.GetFlags(), InFlags))
+			{
+				UE_LOG(LogText3D, Verbose, TEXT("\tFound candidate active builder '%s'."), *CandidateBuilder.GetDebugString());
+				Builder = &CandidateBuilder;
+				break;
+			}
 		}
 	}
 
@@ -72,19 +86,25 @@ void UText3DBuildSystem::RequestTextBuild(TNotNull<UText3DComponent*> InComponen
 	{
 		for (UE::Text3D::FTextBuilder& CandidateBuilder : QueuedBuilders)
 		{
-			if (CanUseBuilder(CandidateBuilder))
+			if (CandidateBuilder.GetComponent() == InComponent && IsBuilderFlagsCompatible(CandidateBuilder))
 			{
+				UE_LOG(LogText3D, Verbose, TEXT("\tFound candidate queued builder '%s'."), *CandidateBuilder.GetDebugString());
 				Builder = &CandidateBuilder;
 				break;
 			}
 		}
 	}
 
-	if (!Builder)
+	if (Builder)
+	{
+		Builder->PrepareFlags(InFlags);
+	}
+	else
 	{
 		Builder = &QueuedBuilders.Emplace_GetRef(InComponent);
+		Builder->PrepareFlags(InFlags);
+		UE_LOG(LogText3D, Verbose, TEXT("\tNo candidate text builder was found. Queueing a new builder '%s'"), *Builder->GetDebugString());
 	}
-	Builder->PrepareFlags(InFlags);
 
 	ConditionallyRegisterDelegates();
 
@@ -133,20 +153,26 @@ void UText3DBuildSystem::IncrementalTextBuild()
 	{
 		// Invalidate the gotten frame cycles
 		BeginFrameCycles.Reset();
+		UE_LOG(LogText3D, Verbose, TEXT("------------------------------------------------------------------"));
 	};
 
+	UE_LOG(LogText3D, Verbose, TEXT("[%llu] Starting incremental Text3D build requirements..."), GFrameCounter);
+
 	TextBuilders.Append(MoveTemp(QueuedBuilders));
 	QueuedBuilders.Reset();
 
 	// Remove all finished builders
+	UE_LOG(LogText3D, Verbose, TEXT("\tCleaning up finished/canceled Text3D builders..."));
 	TextBuilders.RemoveAll(
 		[](const UE::Text3D::FTextBuilder& InBuilder)
 		{
+			UE_CLOG(InBuilder.IsFinished(), LogText3D, Verbose, TEXT("\t\tRemoving finished/canceled Text3D builder '%s'"), *InBuilder.GetDebugString());
 			return InBuilder.IsFinished();
 		});
 
 	if (TextBuilders.IsEmpty())
 	{
+		UE_LOG(LogText3D, Verbose, TEXT("\tIncremental Text3D build finished. No more build jobs left."));
 		FCoreDelegates::OnEndFrame.Remove(EndFrameHandle);
 		EndFrameHandle.Reset();
 		return; // nothing to update
@@ -154,8 +180,11 @@ void UText3DBuildSystem::IncrementalTextBuild()
 
 	const uint64 CurrentCycles = FPlatformTime::Cycles64();
 
+	UE_LOG(LogText3D, Verbose, TEXT("\tChecking Text3D incremental build budget..."));
+
 	// If the incremental build has starved for too long, force it to happen
 	const bool bForceBuild = FramesStarved >= Private::CVarStarveFrameLimit.GetValueOnGameThread();
+	UE_CLOG(bForceBuild, LogText3D, Verbose, TEXT("\t\tIncremental Text3D build starved for more than %d frames. Forcing incremental build..."), Private::CVarStarveFrameLimit.GetValueOnGameThread());
 
 	TOptional<uint64> BudgetEnd;
 	if (UE::Text3D::Private::CVarEnableBuildTimeBudget.GetValueOnGameThread())
@@ -168,7 +197,7 @@ void UText3DBuildSystem::IncrementalTextBuild()
 			{
 				++FramesStarved;
 				BudgetDeficit = 0;
-				UE_LOG(LogText3D, Verbose, TEXT("UText3DBuildSystem::IncrementalTextBuild ran out of build time budget due to blocking builds."));
+				UE_LOG(LogText3D, Verbose, TEXT("\t\tIncremental Text3D build ran out of build time budget due to blocking builds."));
 				return;
 			}
 			*BudgetEnd -= BudgetDeficit;
@@ -177,33 +206,22 @@ void UText3DBuildSystem::IncrementalTextBuild()
 	BudgetDeficit = 0;
 
 	// Skip if the frame duration has already gone for too long
-	const uint64 FrameDuration = CurrentCycles - *BeginFrameCycles;
-	if (!bForceBuild && FrameDuration >= FPlatformTime::SecondsToCycles64(Private::CVarFrameTimeBudget.GetValueOnGameThread()))
+	if (!bForceBuild && (CurrentCycles - *BeginFrameCycles) >= FPlatformTime::SecondsToCycles64(Private::CVarFrameTimeBudget.GetValueOnGameThread()))
 	{
 		++FramesStarved;
-		UE_LOG(LogText3D, Verbose, TEXT("UText3DBuildSystem::IncrementalTextBuild ran out of frame time budget."));
+		UE_LOG(LogText3D, Verbose, TEXT("\t\tIncremental Text3D build ran out of frame time budget."));
 		return;
 	}
 	FramesStarved = 0;
 
-	// Update the builders that are still in the preparation phases
-	for (UE::Text3D::FTextBuilder& Builder : TextBuilders)
-	{
-		const EBuildResult Result = Builder.IncrementalBuild(BudgetEnd, /*MaxPhase*/UE::Text3D::EBuildPhase::PrepareExtensions);
-		if (Result == EBuildResult::OutOfBudget)
-		{
-			UE_LOG(LogText3D, VeryVerbose, TEXT("UText3DBuildSystem::IncrementalTextBuild ran out of budget while building. Continuing next frame"));
-			return; // ran out of budget
-		}
-	}
-
+	UE_LOG(LogText3D, Verbose, TEXT("\tPerforming incremental Text3D build..."));
 	// Update all the builders with no phase limit
 	for (UE::Text3D::FTextBuilder& Builder : TextBuilders)
 	{
 		const EBuildResult Result = Builder.IncrementalBuild(BudgetEnd);
 		if (Result == EBuildResult::OutOfBudget)
 		{
-			UE_LOG(LogText3D, VeryVerbose, TEXT("UText3DBuildSystem::IncrementalTextBuild ran out of budget while building. Continuing next frame"));
+			UE_LOG(LogText3D, Verbose, TEXT("\t\tIncremental Text3D build ran out of budget while building '%s'. Continuing next frame"), *Builder.GetDebugString());
 			return; // ran out of budget
 		}
 	}
diff --git a/Engine/Plugins/VirtualProduction/Text3D/Source/Text3D/Private/Text3DBuilder.cpp b/Engine/Plugins/VirtualProduction/Text3D/Source/Text3D/Private/Text3DBuilder.cpp
--- a/Engine/Plugins/VirtualProduction/Text3D/Source/Text3D/Private/Text3DBuilder.cpp
+++ b/Engine/Plugins/VirtualProduction/Text3D/Source/Text3D/Private/Text3DBuilder.cpp
@@ -24,15 +24,20 @@ bool FTextBuilder::PrepareFlags(EText3DRendererFlags InFlags)
 	return true;
 }
 
-bool FTextBuilder::IsInProgress() const
+FString FTextBuilder::GetDebugString() const
 {
-	return Phase != EBuildPhase::NotStarted 
-		&& Phase != EBuildPhase::Finished;
+	return FString::Printf(TEXT("Component: %s, Build Flags: %d"), *GetFullNameSafe(Component), static_cast<int32>(Flags));
+}
+
+bool FTextBuilder::	IsInProgress() const
+{
+	return Phase != EBuildPhase::NotStarted && !IsFinished();
 }
 
 bool FTextBuilder::IsFinished() const
 {
-	return Phase == EBuildPhase::Finished;
+	return Phase == EBuildPhase::Succeeded
+		|| Phase == EBuildPhase::Canceled;
 }
 
 void FTextBuilder::BlockingBuild()
@@ -70,6 +75,15 @@ EBuildResult FTextBuilder::IncrementalBuild(TOptional<uint64> InBudgetEnd, TOpti
 	return BuildInternal();
 }
 
+void FTextBuilder::CancelBuild()
+{
+	Component->CancelBuild();
+
+	Extensions.Reset();
+	ExtensionIndex = INDEX_NONE;
+	Phase = EBuildPhase::Canceled;
+}
+
 void FTextBuilder::OnObjectsReplaced(const TMap<UObject*, UObject*>& InReplacementMap)
 {
 	if (UObject* const* TargetObject = InReplacementMap.Find(Component))
@@ -127,10 +141,14 @@ EBuildResult FTextBuilder::BuildInternal()
 		Result = PostRendererUpdate();
 		break;
 
-	default:
-	case EBuildPhase::Finished:
+	case EBuildPhase::Succeeded:
 		Finish();
 		return EBuildResult::BuildCompleted;
+
+	default:
+	case EBuildPhase::Canceled:
+		CancelBuild();
+		return EBuildResult::BuildCompleted;
 	}
 
 	// If work is complete, look to the next phase
@@ -352,7 +370,7 @@ void FTextBuilder::Finish()
 
 	Extensions.Reset();
 	ExtensionIndex = INDEX_NONE;
-	Phase = EBuildPhase::Finished;
+	Phase = EBuildPhase::Succeeded;
 }
 
 } // UE::Text3D
diff --git a/Engine/Plugins/VirtualProduction/Text3D/Source/Text3D/Private/Text3DBuilder.h b/Engine/Plugins/VirtualProduction/Text3D/Source/Text3D/Private/Text3DBuilder.h
--- a/Engine/Plugins/VirtualProduction/Text3D/Source/Text3D/Private/Text3DBuilder.h
+++ b/Engine/Plugins/VirtualProduction/Text3D/Source/Text3D/Private/Text3DBuilder.h
@@ -23,8 +23,10 @@ enum class EBuildPhase : uint8
 	RendererUpdate,
 	/** Calls PostRendererUpdate for the Text3D extensions */
 	PostRendererUpdate,
-	/** Build finished */
-	Finished,
+	/** Build finished successfully */
+	Succeeded,
+	/** Build was canceled */
+	Canceled,
 };
 
 /** The build result for incremental builds  */
@@ -71,6 +73,9 @@ public:
 		return Flags;
 	}
 
+	/** Gets a debug string to identify the builder */
+	FString GetDebugString() const;
+
 	/** Returns true if the builder is currently in progress (i.e. between not-started and finished) */
 	bool IsInProgress() const;
 
@@ -83,6 +88,9 @@ public:
 	/** Incrementally builds the component every call with a budget end and update phase */
 	EBuildResult IncrementalBuild(TOptional<uint64> InBudgetEnd, TOptional<UE::Text3D::EBuildPhase> InTempMaxUpdatePhase = {});
 
+	/** Cancels the build */
+	void CancelBuild();
+
 	/** Called by the owning text world subsystem when object replacement has taken place */
 	void OnObjectsReplaced(const TMap<UObject*, UObject*>& InReplacementMap);
 
diff --git a/Engine/Plugins/VirtualProduction/Text3D/Source/Text3D/Private/Text3DComponent.cpp b/Engine/Plugins/VirtualProduction/Text3D/Source/Text3D/Private/Text3DComponent.cpp
--- a/Engine/Plugins/VirtualProduction/Text3D/Source/Text3D/Private/Text3DComponent.cpp
+++ b/Engine/Plugins/VirtualProduction/Text3D/Source/Text3D/Private/Text3DComponent.cpp
@@ -90,18 +90,11 @@ void UText3DComponent::PostLoad()
 
 void UText3DComponent::RequestUpdate(EText3DRendererFlags InFlags, bool bInImmediate)
 {
-	if (InFlags != EText3DRendererFlags::None && !IsTemplate(RF_ClassDefaultObject))
+	if (InFlags != EText3DRendererFlags::None && !IsTemplate(RF_ClassDefaultObject) && IsRegistered())
 	{
-		if (IsRegistered())
+		if (UText3DBuildSystem* const BuildSystem = UText3DBuildSystem::Get(this))
 		{
-			if (UText3DBuildSystem* const BuildSystem = UText3DBuildSystem::Get(this))
-			{
-				BuildSystem->RequestTextBuild(this, InFlags, bInImmediate);
-			}
-		}
-		else
-		{
-			EnumAddFlags(PreRegisterFlags, InFlags);
+			BuildSystem->RequestTextBuild(this, InFlags, bInImmediate);
 		}
 	}
 }
@@ -271,9 +264,7 @@ bool UText3DComponent::IsCompiledFromBlueprint() const
 void UText3DComponent::OnRegister()
 {
 	Super::OnRegister();
-
-	RequestUpdate(PreRegisterFlags);
-	PreRegisterFlags = EText3DRendererFlags::None;
+	RequestUpdate(EText3DRendererFlags::All);
 }
 
 void UText3DComponent::OnUnregister()
@@ -1230,6 +1221,11 @@ void UText3DComponent::FinishBuild(EText3DRendererFlags InFlags)
 	bIsUpdatingText = false;
 }
 
+void UText3DComponent::CancelBuild()
+{
+	bIsUpdatingText = false;
+}
+
 void UText3DComponent::OnVisibilityChanged()
 {
 	Super::OnVisibilityChanged();
diff --git a/Engine/Plugins/VirtualProduction/Text3D/Source/Text3D/Private/Text3DTypes.cpp b/Engine/Plugins/VirtualProduction/Text3D/Source/Text3D/Private/Text3DTypes.cpp
--- a/Engine/Plugins/VirtualProduction/Text3D/Source/Text3D/Private/Text3DTypes.cpp
+++ b/Engine/Plugins/VirtualProduction/Text3D/Source/Text3D/Private/Text3DTypes.cpp
@@ -17,6 +17,16 @@ namespace UE::Text3D::Geometry
 		}
 	}
 
+	FCachedFontFaceGlyphHandle::FCachedFontFaceGlyphHandle(const FCachedFontFaceGlyphHandle& InOther)
+	{
+		*this = InOther;
+	}
+
+	FCachedFontFaceGlyphHandle::FCachedFontFaceGlyphHandle(FCachedFontFaceGlyphHandle&& InOther)
+	{
+		*this = MoveTemp(InOther);
+	}
+
 	FCachedFontFaceGlyphHandle::~FCachedFontFaceGlyphHandle()
 	{
 		Unset();
@@ -24,13 +34,31 @@ namespace UE::Text3D::Geometry
 
 	FCachedFontFaceGlyphHandle& FCachedFontFaceGlyphHandle::operator=(const FCachedFontFaceGlyphHandle& InOther)
 	{
-		Unset();
-		FontFaceHash = InOther.FontFaceHash;
-		GlyphIndex = InOther.GlyphIndex;
-		GlyphMeshHash = InOther.GlyphMeshHash;
-		if (FText3DCachedMesh* CachedMesh = Resolve())
+		if (this != &InOther)
 		{
-			CachedMesh->RefCount++;
+			Unset();
+			FontFaceHash = InOther.FontFaceHash;
+			GlyphIndex = InOther.GlyphIndex;
+			GlyphMeshHash = InOther.GlyphMeshHash;
+			if (FText3DCachedMesh* CachedMesh = Resolve())
+			{
+				CachedMesh->RefCount++;
+			}
+		}
+		return *this;
+	}
+
+	FCachedFontFaceGlyphHandle& FCachedFontFaceGlyphHandle::operator=(FCachedFontFaceGlyphHandle&& InOther)
+	{
+		if (this != &InOther)
+		{
+			Unset();
+			FontFaceHash = InOther.FontFaceHash;
+			GlyphIndex = InOther.GlyphIndex;
+			GlyphMeshHash = InOther.GlyphMeshHash;
+			InOther.FontFaceHash = 0;
+			InOther.GlyphIndex = 0;
+			InOther.GlyphMeshHash = 0;
 		}
 		return *this;
 	}
diff --git a/Engine/Plugins/VirtualProduction/Text3D/Source/Text3D/Public/Text3DComponent.h b/Engine/Plugins/VirtualProduction/Text3D/Source/Text3D/Public/Text3DComponent.h
--- a/Engine/Plugins/VirtualProduction/Text3D/Source/Text3D/Public/Text3DComponent.h
+++ b/Engine/Plugins/VirtualProduction/Text3D/Source/Text3D/Public/Text3DComponent.h
@@ -424,6 +424,9 @@ public:
 	/** Called when the text has finished building */
 	void FinishBuild(EText3DRendererFlags InFlags);
 
+	/** Called when the build is cancelled */
+	void CancelBuild();
+
 protected:
 	//~ Begin USceneComponent
 	/** Intercept and propagate a change on this component to all children. */
@@ -554,9 +557,6 @@ private:
 	/** Called before text is updated */
 	FTextUpdated TextPreUpdateDelegate;
 
-	/** Flags from requests prior to component registration */
-	EText3DRendererFlags PreRegisterFlags = EText3DRendererFlags::None;
-
 	/** Called after text is updated */
 	FTextUpdated TextPostUpdateDelegate;
 
diff --git a/Engine/Plugins/VirtualProduction/Text3D/Source/Text3D/Public/Text3DTypes.h b/Engine/Plugins/VirtualProduction/Text3D/Source/Text3D/Public/Text3DTypes.h
--- a/Engine/Plugins/VirtualProduction/Text3D/Source/Text3D/Public/Text3DTypes.h
+++ b/Engine/Plugins/VirtualProduction/Text3D/Source/Text3D/Public/Text3DTypes.h
@@ -391,9 +391,12 @@ namespace UE::Text3D
 		{
 			FCachedFontFaceGlyphHandle() = default;
 			FCachedFontFaceGlyphHandle(const FText3DFontFaceCache& InFontFaceCache, uint32 InGlyphIndex, const FGlyphMeshParameters& InGlyphMeshParameters);
+			FCachedFontFaceGlyphHandle(const FCachedFontFaceGlyphHandle& InOther);
+			FCachedFontFaceGlyphHandle(FCachedFontFaceGlyphHandle&& InOther);
 			~FCachedFontFaceGlyphHandle();
 
 			FCachedFontFaceGlyphHandle& operator=(const FCachedFontFaceGlyphHandle& InOther);
+			FCachedFontFaceGlyphHandle& operator=(FCachedFontFaceGlyphHandle&& InOther);
 
 			void Unset();
 
