Text3D: added scene view extension to hide primitives that are being built to avoid visual errors when text is being rebuilt.
- Improved build logging
- Fixed issues with text3d builds being cancelled when modifying properties

#Synced-CL 51290016

diff --git a/Engine/Plugins/VirtualProduction/Text3D/Source/Text3D/Private/Renderers/DynamicMesh/Text3DDynamicMeshRenderer.cpp b/Engine/Plugins/VirtualProduction/Text3D/Source/Text3D/Private/Renderers/DynamicMesh/Text3DDynamicMeshRenderer.cpp
--- a/Engine/Plugins/VirtualProduction/Text3D/Source/Text3D/Private/Renderers/DynamicMesh/Text3DDynamicMeshRenderer.cpp
+++ b/Engine/Plugins/VirtualProduction/Text3D/Source/Text3D/Private/Renderers/DynamicMesh/Text3DDynamicMeshRenderer.cpp
@@ -228,6 +228,14 @@ FBox UText3DDynamicMeshRenderer::OnCalculateBounds() const
 	return Box;
 }
 
+void UText3DDynamicMeshRenderer::OnIterateManagedPrimitives(TFunctionRef<void(TNotNull<const UPrimitiveComponent*>)> InFunc) const
+{
+	if (DynamicMeshComponent)
+	{
+		InFunc(DynamicMeshComponent);
+	}
+}
+
 #if WITH_EDITOR
 void UText3DDynamicMeshRenderer::OnDebugModeEnabled()
 {
diff --git a/Engine/Plugins/VirtualProduction/Text3D/Source/Text3D/Private/Renderers/DynamicMesh/Text3DDynamicMeshRenderer.h b/Engine/Plugins/VirtualProduction/Text3D/Source/Text3D/Private/Renderers/DynamicMesh/Text3DDynamicMeshRenderer.h
--- a/Engine/Plugins/VirtualProduction/Text3D/Source/Text3D/Private/Renderers/DynamicMesh/Text3DDynamicMeshRenderer.h
+++ b/Engine/Plugins/VirtualProduction/Text3D/Source/Text3D/Private/Renderers/DynamicMesh/Text3DDynamicMeshRenderer.h
@@ -63,6 +63,7 @@ protected:
 	virtual EText3DMeshType GetMeshType() const override;
 	virtual FName GetFriendlyName() const override;
 	virtual FBox OnCalculateBounds() const override;
+	virtual void OnIterateManagedPrimitives(TFunctionRef<void(TNotNull<const UPrimitiveComponent*>)> InFunc) const override;
 #if WITH_EDITOR
 	virtual void OnDebugModeEnabled() override;
 	virtual void OnDebugModeDisabled() override;
diff --git a/Engine/Plugins/VirtualProduction/Text3D/Source/Text3D/Private/Renderers/StaticMeshes/Text3DStaticMeshesRenderer.cpp b/Engine/Plugins/VirtualProduction/Text3D/Source/Text3D/Private/Renderers/StaticMeshes/Text3DStaticMeshesRenderer.cpp
--- a/Engine/Plugins/VirtualProduction/Text3D/Source/Text3D/Private/Renderers/StaticMeshes/Text3DStaticMeshesRenderer.cpp
+++ b/Engine/Plugins/VirtualProduction/Text3D/Source/Text3D/Private/Renderers/StaticMeshes/Text3DStaticMeshesRenderer.cpp
@@ -14,6 +14,7 @@
 #include "SceneInterface.h"
 #include "Text3DComponent.h"
 #include "Text3DInternalTypes.h"
+#include "UObject/Package.h"
 #include "Utilities/Text3DUtilities.h"
 
 #if WITH_EDITOR
@@ -44,6 +45,7 @@ void UText3DStaticMeshesRenderer::OnCreate()
 	if (!TextRoot)
 	{
 		TextRoot = NewObject<USceneComponent>(this, TEXT("TextRoot"));
+		TextRoot->CreationMethod = EComponentCreationMethod::Instance;
 		TextRoot->RegisterComponent();
 		Owner->AddOwnedComponent(TextRoot);
 		TextRoot->AttachToComponent(TextComponent, AttachRule);
@@ -55,9 +57,13 @@ void UText3DStaticMeshesRenderer::OnCreate()
 		TextRoot->DetachFromComponent(DetachRule);
 		TextRoot->AttachToComponent(TextComponent, AttachRule);
 
-		if (TextRoot->GetOuter() != this)
+		if (TextRoot->GetOwner() != Owner)
 		{
-			TextRoot->Rename(nullptr, this, REN_DontCreateRedirectors | REN_NonTransactional | REN_DoNotDirty);
+			// Temporarily move the root to the transient package to retrigger UActorComponent::PostRename which updates
+			// OwnerPrivate only when the outers have changed.
+			constexpr ERenameFlags RenameFlags = REN_DontCreateRedirectors | REN_NonTransactional | REN_DoNotDirty;
+			TextRoot->Rename(nullptr, GetTransientPackage(), RenameFlags);
+			TextRoot->Rename(nullptr, this, RenameFlags);
 		}
 	}
 
@@ -317,6 +323,17 @@ FBox UText3DStaticMeshesRenderer::OnCalculateBounds() const
 	return Box;
 }
 
+void UText3DStaticMeshesRenderer::OnIterateManagedPrimitives(TFunctionRef<void(TNotNull<const UPrimitiveComponent*>)> InFunc) const
+{
+	for (const UStaticMeshComponent* CharacterMesh : CharacterMeshes)
+	{
+		if (CharacterMesh)
+		{
+			InFunc(CharacterMesh);
+		}
+	}
+}
+
 void UText3DStaticMeshesRenderer::AllocateComponents(int32 InCount)
 {
 	CharacterMeshes.RemoveAll([](const UStaticMeshComponent* InComponent)
diff --git a/Engine/Plugins/VirtualProduction/Text3D/Source/Text3D/Private/Renderers/StaticMeshes/Text3DStaticMeshesRenderer.h b/Engine/Plugins/VirtualProduction/Text3D/Source/Text3D/Private/Renderers/StaticMeshes/Text3DStaticMeshesRenderer.h
--- a/Engine/Plugins/VirtualProduction/Text3D/Source/Text3D/Private/Renderers/StaticMeshes/Text3DStaticMeshesRenderer.h
+++ b/Engine/Plugins/VirtualProduction/Text3D/Source/Text3D/Private/Renderers/StaticMeshes/Text3DStaticMeshesRenderer.h
@@ -53,6 +53,7 @@ protected:
 	virtual EText3DMeshType GetMeshType() const override;
 	virtual FName GetFriendlyName() const override;
 	virtual FBox OnCalculateBounds() const override;
+	virtual void OnIterateManagedPrimitives(TFunctionRef<void(TNotNull<const UPrimitiveComponent*>)> InFunc) const override;
 #if WITH_EDITOR
 	virtual void OnDebugModeEnabled() override;
 	virtual void OnDebugModeDisabled() override;
diff --git a/Engine/Plugins/VirtualProduction/Text3D/Source/Text3D/Private/Renderers/Text3DRendererBase.cpp b/Engine/Plugins/VirtualProduction/Text3D/Source/Text3D/Private/Renderers/Text3DRendererBase.cpp
--- a/Engine/Plugins/VirtualProduction/Text3D/Source/Text3D/Private/Renderers/Text3DRendererBase.cpp
+++ b/Engine/Plugins/VirtualProduction/Text3D/Source/Text3D/Private/Renderers/Text3DRendererBase.cpp
@@ -155,3 +155,8 @@ FBox UText3DRendererBase::GetBounds() const
 
 	return CachedBounds.GetValue();
 }
+
+void UText3DRendererBase::IterateManagedPrimitives(TFunctionRef<void(TNotNull<const UPrimitiveComponent*>)> InFunc) const
+{
+	OnIterateManagedPrimitives(InFunc);
+}
diff --git a/Engine/Plugins/VirtualProduction/Text3D/Source/Text3D/Private/Subsystems/Text3DBuildSystem.cpp b/Engine/Plugins/VirtualProduction/Text3D/Source/Text3D/Private/Subsystems/Text3DBuildSystem.cpp
--- a/Engine/Plugins/VirtualProduction/Text3D/Source/Text3D/Private/Subsystems/Text3DBuildSystem.cpp
+++ b/Engine/Plugins/VirtualProduction/Text3D/Source/Text3D/Private/Subsystems/Text3DBuildSystem.cpp
@@ -1,17 +1,33 @@
 // Copyright Epic Games, Inc. All Rights Reserved.
 
 #include "Text3DBuildSystem.h"
+#include "Algo/Contains.h"
+#include "Algo/ForEach.h"
+#include "Components/PrimitiveComponent.h"
 #include "Engine/World.h"
 #include "Logs/Text3DLogs.h"
 #include "Misc/CoreDelegates.h"
+#include "Renderers/Text3DRendererBase.h"
+#include "SceneView.h"
+#include "SceneViewExtension.h"
 #include "Text3DComponent.h"
+#include "Text3DSceneViewExtension.h"
 
 namespace UE::Text3D::Private
 {
+TAutoConsoleVariable<int32> CVarFrameTimeBudgetMode(TEXT("Text3D.Build.FrameTimeBudgetMode")
+	, 1
+	, TEXT("The frame time budget mode to build Text3D components.\n")
+	  TEXT("0: Frame time budget disabled\n")
+	  TEXT("1: Frame time budget enabled only in playing worlds\n")
+	  TEXT("2: Frame time budget enabled in both playing worlds and editor worlds.\n")
+	);
+
 TAutoConsoleVariable<float> CVarFrameTimeBudget(TEXT("Text3D.Build.FrameTimeBudget")
 	, 0.016f
-	, TEXT("The frame time budget to build Text3D components.\n")
-	  TEXT("If the frame is already past this, the build will be skipped unless it has starved for a while.\n")
+	, TEXT("The frame time cap to build Text3D components.\n")
+	  TEXT("Only relevant if Text3DBuild.FrameTimeBudgetMode is non-zero.\n")
+	  TEXT("If the current frame duration prior to the incremental build is past this limit, the build will be skipped (unless it has starved for more than the starve limit).\n")
 	  TEXT("See Text3D.Build.StarveFrameLimit")
 	);
 
@@ -21,11 +37,51 @@ TAutoConsoleVariable<int32> CVarStarveFrameLimit(TEXT("Text3D.Build.StarveFrameL
 
 TAutoConsoleVariable<bool> CVarEnableBuildTimeBudget(TEXT("Text3D.Build.EnableBuildTimeBudget")
 	, true
-	, TEXT("Whether to have time budget to build Text3D components."));
+	, TEXT("Whether to have a time budget when building Text3D components.\n"));
 
 TAutoConsoleVariable<float> CVarBuildTimeBudget(TEXT("Text3D.Build.BuildTimeBudget")
 	, 0.002f
-	, TEXT("The time budget to build Text3D components."));
+	, TEXT("The time budget allocated for building Text3D components.\n")
+	  TEXT("Requires Text3D.Build.EnableBuildTimeBudget to be true."));
+
+/** Returns true if time budget is enabled for the given world */
+bool IsFrameTimeBudgetEnabled(const UWorld* InWorld)
+{
+	if (!InWorld)
+	{
+		return false;
+	}
+
+	// Mode 0: Disabled
+	// Mode 1: Only playing worlds
+	// Mode 2: Any world
+	const int32 Mode = CVarFrameTimeBudgetMode.GetValueOnGameThread();
+	if (Mode <= 0 || Mode > 2)
+	{
+		return false;
+	}
+
+#if WITH_EDITOR
+	constexpr EWorldType::Type PlayingWorldTypes[] = { EWorldType::Game, EWorldType::GamePreview, EWorldType::PIE };
+	const bool bOnlyPlayingWorlds = Mode == 1;
+	return !bOnlyPlayingWorlds || Algo::Contains(PlayingWorldTypes, InWorld->WorldType);
+#else
+	return true;
+#endif
+}
+
+/** Returns the scene view extension to use for this build system */
+TSharedRef<FText3DSceneViewExtension> GetSceneViewExtension()
+{
+	static TWeakPtr<FText3DSceneViewExtension> SceneViewExtensionWeak;
+	TSharedPtr<FText3DSceneViewExtension> SceneViewExtension = SceneViewExtensionWeak.Pin();
+	if (!SceneViewExtension.IsValid())
+	{
+		SceneViewExtension = FSceneViewExtensions::NewExtension<FText3DSceneViewExtension>();
+		SceneViewExtensionWeak = SceneViewExtension;
+	}
+	return SceneViewExtension.ToSharedRef();
+}
 
 } // UE::Text3D::Private
 
@@ -45,11 +101,11 @@ void UText3DBuildSystem::RequestTextBuild(TNotNull<UText3DComponent*> InComponen
 		return;
 	}
 
-	UE_LOG(LogText3D, Verbose, TEXT("[%llu] Requesting text build for component '%s'. Flags: %d. Blocking? %d")
+	UE_LOG(LogText3D, Verbose, TEXT("[%llu] Requesting %s text build for component '%s'. Flags: %d")
 		, GFrameCounter
+		, bInBlockingBuild ? TEXT("BLOCKING") : TEXT("incremental")
 		, *InComponent->GetFullName()
-		, static_cast<int32>(InFlags)
-		, bInBlockingBuild);
+		, static_cast<int32>(InFlags));
 
 	// allow if build is incremental. If blocking ensure the flags are exactly equal. If there's another flag, it can't be used.
 	auto IsBuilderFlagsCompatible = 
@@ -119,22 +175,20 @@ void UText3DBuildSystem::RequestTextBuild(TNotNull<UText3DComponent*> InComponen
 	}
 }
 
-void UText3DBuildSystem::RemoveTextBuild(TNotNull<UText3DComponent*> InComponent)
+int32 UText3DBuildSystem::RemoveTextBuild(TNotNull<UText3DComponent*> InComponent)
 {
-	for (int32 Index = QueuedBuilders.Num() - 1; Index >= 0; --Index)
-	{
-		if (QueuedBuilders[Index].GetComponent() == InComponent)
-		{
-			QueuedBuilders.RemoveAt(Index);
-		}
-	}
-	for (int32 Index = TextBuilders.Num() - 1; Index >= 0; --Index)
-	{
-		if (TextBuilders[Index].GetComponent() == InComponent)
+	auto IsMatchingComponent = 
+		[InComponent](const UE::Text3D::FTextBuilder& InTextBuilder)
 		{
-			TextBuilders.RemoveAt(Index);
-		}
-	}
+			return InTextBuilder.GetComponent() == InComponent;
+		};
+
+	int32 RemovedCount = 0;
+	RemovedCount += QueuedBuilders.RemoveAll(IsMatchingComponent);
+	RemovedCount += TextBuilders.RemoveAll(IsMatchingComponent);
+
+	UE_CLOG(RemovedCount > 0, LogText3D, Verbose, TEXT("Removed '%s' from %d builders"), *InComponent->GetFullName(), RemovedCount);
+	return RemovedCount;
 }
 
 void UText3DBuildSystem::IncrementalTextBuild()
@@ -153,21 +207,19 @@ void UText3DBuildSystem::IncrementalTextBuild()
 	{
 		// Invalidate the gotten frame cycles
 		BeginFrameCycles.Reset();
-		UE_LOG(LogText3D, Verbose, TEXT("------------------------------------------------------------------"));
 	};
 
-	UE_LOG(LogText3D, Verbose, TEXT("[%llu] Starting incremental Text3D build requirements..."), GFrameCounter);
+	UE_LOG(LogText3D, Verbose, TEXT("[%llu] Performing incremental Text3D build..."), GFrameCounter);
 
 	TextBuilders.Append(MoveTemp(QueuedBuilders));
 	QueuedBuilders.Reset();
 
 	// Remove all finished builders
-	UE_LOG(LogText3D, Verbose, TEXT("\tCleaning up finished/canceled Text3D builders..."));
-	TextBuilders.RemoveAll(
+	const int32 RemovedCount = TextBuilders.RemoveAll(
 		[](const UE::Text3D::FTextBuilder& InBuilder)
 		{
-			UE_CLOG(InBuilder.IsFinished(), LogText3D, Verbose, TEXT("\t\tRemoving finished/canceled Text3D builder '%s'"), *InBuilder.GetDebugString());
-			return InBuilder.IsFinished();
+			UE_CLOG(InBuilder.IsComplete(), LogText3D, Verbose, TEXT("\tRemoving finished Text3D builders '%s'"), *InBuilder.GetDebugString());
+			return InBuilder.IsComplete();
 		});
 
 	if (TextBuilders.IsEmpty())
@@ -180,11 +232,9 @@ void UText3DBuildSystem::IncrementalTextBuild()
 
 	const uint64 CurrentCycles = FPlatformTime::Cycles64();
 
-	UE_LOG(LogText3D, Verbose, TEXT("\tChecking Text3D incremental build budget..."));
-
 	// If the incremental build has starved for too long, force it to happen
 	const bool bForceBuild = FramesStarved >= Private::CVarStarveFrameLimit.GetValueOnGameThread();
-	UE_CLOG(bForceBuild, LogText3D, Verbose, TEXT("\t\tIncremental Text3D build starved for more than %d frames. Forcing incremental build..."), Private::CVarStarveFrameLimit.GetValueOnGameThread());
+	UE_CLOG(bForceBuild, LogText3D, Verbose, TEXT("\tIncremental Text3D build starved for more than %d frames. Forcing incremental build..."), Private::CVarStarveFrameLimit.GetValueOnGameThread());
 
 	TOptional<uint64> BudgetEnd;
 	if (UE::Text3D::Private::CVarEnableBuildTimeBudget.GetValueOnGameThread())
@@ -197,7 +247,7 @@ void UText3DBuildSystem::IncrementalTextBuild()
 			{
 				++FramesStarved;
 				BudgetDeficit = 0;
-				UE_LOG(LogText3D, Verbose, TEXT("\t\tIncremental Text3D build ran out of build time budget due to blocking builds."));
+				UE_LOG(LogText3D, Verbose, TEXT("\tIncremental Text3D build ran out of build time budget due to blocking builds."));
 				return;
 			}
 			*BudgetEnd -= BudgetDeficit;
@@ -206,15 +256,14 @@ void UText3DBuildSystem::IncrementalTextBuild()
 	BudgetDeficit = 0;
 
 	// Skip if the frame duration has already gone for too long
-	if (!bForceBuild && (CurrentCycles - *BeginFrameCycles) >= FPlatformTime::SecondsToCycles64(Private::CVarFrameTimeBudget.GetValueOnGameThread()))
+	if (!bForceBuild && Private::IsFrameTimeBudgetEnabled(GetWorld()) && (CurrentCycles - *BeginFrameCycles) >= FPlatformTime::SecondsToCycles64(Private::CVarFrameTimeBudget.GetValueOnGameThread()))
 	{
 		++FramesStarved;
-		UE_LOG(LogText3D, Verbose, TEXT("\t\tIncremental Text3D build ran out of frame time budget."));
+		UE_LOG(LogText3D, Verbose, TEXT("\tIncremental Text3D build ran out of frame time budget."));
 		return;
 	}
 	FramesStarved = 0;
 
-	UE_LOG(LogText3D, Verbose, TEXT("\tPerforming incremental Text3D build..."));
 	// Update all the builders with no phase limit
 	for (UE::Text3D::FTextBuilder& Builder : TextBuilders)
 	{
@@ -231,18 +280,18 @@ void UText3DBuildSystem::Initialize(FSubsystemCollectionBase& InCollection)
 {
 	Super::Initialize(InCollection);
 
+	SceneViewExtension = UE::Text3D::Private::GetSceneViewExtension();
+
 #if WITH_EDITOR
 	FCoreUObjectDelegates::OnObjectsReplaced.AddWeakLambda(this, 
 		[this](const TMap<UObject*, UObject*>& InReplacementMap)
 		{
-			for (UE::Text3D::FTextBuilder& TextBuilder : QueuedBuilders)
-			{
-				TextBuilder.OnObjectsReplaced(InReplacementMap);	
-			}
-			for (UE::Text3D::FTextBuilder& TextBuilder : TextBuilders)
-			{
-				TextBuilder.OnObjectsReplaced(InReplacementMap);	
-			}
+			auto ReplaceObjects = [&InReplacementMap](UE::Text3D::FTextBuilder& InBuilder)
+				{
+					InBuilder.OnObjectsReplaced(InReplacementMap);
+				};
+			Algo::ForEach(QueuedBuilders, ReplaceObjects);
+			Algo::ForEach(TextBuilders, ReplaceObjects);
 		});
 #endif
 }
@@ -254,14 +303,12 @@ void UText3DBuildSystem::Deinitialize()
 	FCoreUObjectDelegates::OnObjectsReplaced.RemoveAll(this);
 #endif
 	UnregisterDelegates();
+	SceneViewExtension.Reset();
 }
 
 bool UText3DBuildSystem::DoesSupportWorldType(const EWorldType::Type InWorldType) const
 {
-	return InWorldType == EWorldType::Game
-		|| InWorldType == EWorldType::Editor
-		|| InWorldType == EWorldType::PIE
-		|| InWorldType == EWorldType::GamePreview;
+	return InWorldType != EWorldType::None;
 }
 
 bool UText3DBuildSystem::IsPendingBuildCompletion() const
@@ -270,10 +317,39 @@ bool UText3DBuildSystem::IsPendingBuildCompletion() const
 	return !QueuedBuilders.IsEmpty() || TextBuilders.ContainsByPredicate(
 		[](const UE::Text3D::FTextBuilder& InTextBuilder)
 		{
-			return !InTextBuilder.IsFinished();
+			return !InTextBuilder.IsComplete();
 		});
 }
 
+void UText3DBuildSystem::HidePrimitivesBeingBuilt(FSceneView& InSceneView) const
+{
+	auto AddPendingBuildItem = 
+		[&InSceneView](const UE::Text3D::FTextBuilder& InBuilder)
+		{
+			// Only consider builders that are building geometry and during or past the renderer update stage.
+			const bool bBuildingGeometryInProgress = EnumHasAnyFlags(InBuilder.GetFlags(), EText3DRendererFlags::Geometry)
+				&& InBuilder.GetPhase() >= UE::Text3D::EBuildPhase::RendererUpdate;
+
+			if (!bBuildingGeometryInProgress)
+			{
+				return;
+			}
+			if (const UText3DComponent* const Component = InBuilder.GetComponent())
+			{
+				if (const UText3DRendererBase* const TextRenderer = Component->GetTextRenderer())
+				{
+					TextRenderer->IterateManagedPrimitives(
+						[&InSceneView](TNotNull<const UPrimitiveComponent*> InComponent)
+						{
+							InSceneView.HiddenPrimitives.Add(InComponent->GetPrimitiveSceneId());
+						});
+				}
+			}
+		};
+	Algo::ForEach(QueuedBuilders, AddPendingBuildItem);
+	Algo::ForEach(TextBuilders, AddPendingBuildItem);
+}
+
 void UText3DBuildSystem::ConditionallyRegisterDelegates()
 {
 	if (!BeginFrameHandle.IsValid())
diff --git a/Engine/Plugins/VirtualProduction/Text3D/Source/Text3D/Private/Subsystems/Text3DBuildSystem.h b/Engine/Plugins/VirtualProduction/Text3D/Source/Text3D/Private/Subsystems/Text3DBuildSystem.h
--- a/Engine/Plugins/VirtualProduction/Text3D/Source/Text3D/Private/Subsystems/Text3DBuildSystem.h
+++ b/Engine/Plugins/VirtualProduction/Text3D/Source/Text3D/Private/Subsystems/Text3DBuildSystem.h
@@ -8,6 +8,11 @@
 #include "Text3DBuilder.h"
 #include "Text3DBuildSystem.generated.h"
 
+namespace UE::Text3D
+{
+	class FText3DSceneViewExtension;
+}
+
 class UText3DComponent;
 enum class EText3DRendererFlags : uint8;
 
@@ -32,7 +37,7 @@ public:
 	void RequestTextBuild(TNotNull<UText3DComponent*> InComponent, EText3DRendererFlags InFlags, bool bInBlockingBuild);
 
 	/** Removes the builders that relate to the given component*/
-	void RemoveTextBuild(TNotNull<UText3DComponent*> InComponent);
+	int32 RemoveTextBuild(TNotNull<UText3DComponent*> InComponent);
 
 	/** Called on tick to process builders and call IncrementalBuild on them */
 	void IncrementalTextBuild();
@@ -49,6 +54,7 @@ protected:
 
 	//~ Begin IText3DBuildSystemInterface
 	virtual bool IsPendingBuildCompletion() const override;
+	virtual void HidePrimitivesBeingBuilt(FSceneView& InSceneView) const override;
 	//~ End IText3DBuildSystemInterface
 
 private:
@@ -64,6 +70,9 @@ private:
 	/** Active builders */
 	TArray<UE::Text3D::FTextBuilder> TextBuilders;
 
+	/** Scene view extension to hide any primitive from text actors that are pending build */
+	TSharedPtr<UE::Text3D::FText3DSceneViewExtension> SceneViewExtension;
+
 	/** Cycles measured in begin frame */
 	TOptional<uint64> BeginFrameCycles;
 
diff --git a/Engine/Plugins/VirtualProduction/Text3D/Source/Text3D/Private/Text3DBuilder.cpp b/Engine/Plugins/VirtualProduction/Text3D/Source/Text3D/Private/Text3DBuilder.cpp
--- a/Engine/Plugins/VirtualProduction/Text3D/Source/Text3D/Private/Text3DBuilder.cpp
+++ b/Engine/Plugins/VirtualProduction/Text3D/Source/Text3D/Private/Text3DBuilder.cpp
@@ -9,6 +9,27 @@
 namespace UE::Text3D
 {
 
+namespace Private
+{
+
+const TCHAR* GetPhaseString(EBuildPhase InPhase)
+{
+	switch (InPhase)
+	{
+	case EBuildPhase::NotStarted:         return TEXT("NotStarted");
+	case EBuildPhase::PrepareExtensions:  return TEXT("PrepareExtensions");
+	case EBuildPhase::PreRendererUpdate:  return TEXT("PreRendererUpdate");
+	case EBuildPhase::RendererUpdate:     return TEXT("RendererUpdate");
+	case EBuildPhase::PostRendererUpdate: return TEXT("PrepareExtensions");
+	case EBuildPhase::Succeeded:          return TEXT("Succeeded");
+	case EBuildPhase::Canceled:           return TEXT("Canceled");
+	}
+	checkNoEntry();
+	return nullptr;
+}
+
+} // UE::Text3D::Private
+
 FTextBuilder::FTextBuilder(UText3DComponent* InComponent)
 	: Component(InComponent)
 {
@@ -26,15 +47,18 @@ bool FTextBuilder::PrepareFlags(EText3DRendererFlags InFlags)
 
 FString FTextBuilder::GetDebugString() const
 {
-	return FString::Printf(TEXT("Component: %s, Build Flags: %d"), *GetFullNameSafe(Component), static_cast<int32>(Flags));
+	return FString::Printf(TEXT("Build Flags %0.2d | Phase %s | Component %s")
+		, static_cast<int32>(Flags)
+		, Private::GetPhaseString(Phase)
+		, *GetFullNameSafe(Component));
 }
 
 bool FTextBuilder::	IsInProgress() const
 {
-	return Phase != EBuildPhase::NotStarted && !IsFinished();
+	return Phase != EBuildPhase::NotStarted && !IsComplete();
 }
 
-bool FTextBuilder::IsFinished() const
+bool FTextBuilder::IsComplete() const
 {
 	return Phase == EBuildPhase::Succeeded
 		|| Phase == EBuildPhase::Canceled;
@@ -51,7 +75,7 @@ void FTextBuilder::BlockingBuild()
 
 EBuildResult FTextBuilder::IncrementalBuild(TOptional<uint64> InBudgetEnd, TOptional<EBuildPhase> InTempMaxUpdatePhase)
 {
-	if (IsFinished())
+	if (IsComplete())
 	{
 		return EBuildResult::BuildCompleted;
 	}
diff --git a/Engine/Plugins/VirtualProduction/Text3D/Source/Text3D/Private/Text3DBuilder.h b/Engine/Plugins/VirtualProduction/Text3D/Source/Text3D/Private/Text3DBuilder.h
--- a/Engine/Plugins/VirtualProduction/Text3D/Source/Text3D/Private/Text3DBuilder.h
+++ b/Engine/Plugins/VirtualProduction/Text3D/Source/Text3D/Private/Text3DBuilder.h
@@ -23,7 +23,7 @@ enum class EBuildPhase : uint8
 	RendererUpdate,
 	/** Calls PostRendererUpdate for the Text3D extensions */
 	PostRendererUpdate,
-	/** Build finished successfully */
+	/** Build completed successfully */
 	Succeeded,
 	/** Build was canceled */
 	Canceled,
@@ -73,14 +73,19 @@ public:
 		return Flags;
 	}
 
+	EBuildPhase GetPhase() const
+	{
+		return Phase;
+	}
+
 	/** Gets a debug string to identify the builder */
 	FString GetDebugString() const;
 
-	/** Returns true if the builder is currently in progress (i.e. between not-started and finished) */
+	/** Returns true if the builder is currently in progress (i.e. between not-started and completed) */
 	bool IsInProgress() const;
 
-	/** Returns true if the builder has finished building the component */
-	bool IsFinished() const;
+	/** Returns true if the builder has completed building the component */
+	bool IsComplete() const;
 
 	/** Fully starts / resumes and finishes the build */
 	void BlockingBuild();
diff --git a/Engine/Plugins/VirtualProduction/Text3D/Source/Text3D/Private/Text3DComponent.cpp b/Engine/Plugins/VirtualProduction/Text3D/Source/Text3D/Private/Text3DComponent.cpp
--- a/Engine/Plugins/VirtualProduction/Text3D/Source/Text3D/Private/Text3DComponent.cpp
+++ b/Engine/Plugins/VirtualProduction/Text3D/Source/Text3D/Private/Text3DComponent.cpp
@@ -90,11 +90,18 @@ void UText3DComponent::PostLoad()
 
 void UText3DComponent::RequestUpdate(EText3DRendererFlags InFlags, bool bInImmediate)
 {
-	if (InFlags != EText3DRendererFlags::None && !IsTemplate(RF_ClassDefaultObject) && IsRegistered())
+	if (InFlags != EText3DRendererFlags::None && !IsTemplate(RF_ClassDefaultObject))
 	{
-		if (UText3DBuildSystem* const BuildSystem = UText3DBuildSystem::Get(this))
+		if (IsRegistered())
 		{
-			BuildSystem->RequestTextBuild(this, InFlags, bInImmediate);
+			if (UText3DBuildSystem* const BuildSystem = UText3DBuildSystem::Get(this))
+			{
+				BuildSystem->RequestTextBuild(this, InFlags, bInImmediate);
+			}
+		}
+		else
+		{
+			EnumAddFlags(PreRegisterUpdateFlags, InFlags);
 		}
 	}
 }
@@ -264,17 +271,12 @@ bool UText3DComponent::IsCompiledFromBlueprint() const
 void UText3DComponent::OnRegister()
 {
 	Super::OnRegister();
-	RequestUpdate(EText3DRendererFlags::All);
+	RequestUpdate(PreRegisterUpdateFlags);
 }
 
 void UText3DComponent::OnUnregister()
 {
 	Super::OnUnregister();
-
-	if (UText3DBuildSystem* BuildSystem = UText3DBuildSystem::Get(this))
-	{
-		BuildSystem->RemoveTextBuild(this);
-	}
 }
 
 void UText3DComponent::PostApplyToComponent()
@@ -1217,6 +1219,7 @@ void UText3DComponent::FinishBuild(EText3DRendererFlags InFlags)
 		MarkRenderStateDirty();
 	}
 
+	EnumRemoveFlags(PreRegisterUpdateFlags, InFlags);
 	TextPostUpdateDelegate.Broadcast(this, InFlags);
 	bIsUpdatingText = false;
 }
diff --git a/Engine/Plugins/VirtualProduction/Text3D/Source/Text3D/Private/Text3DSceneViewExtension.cpp b/Engine/Plugins/VirtualProduction/Text3D/Source/Text3D/Private/Text3DSceneViewExtension.cpp
new file mode 100644
--- /dev/null
+++ b/Engine/Plugins/VirtualProduction/Text3D/Source/Text3D/Private/Text3DSceneViewExtension.cpp
@@ -0,0 +1,29 @@
+// Copyright Epic Games, Inc. All Rights Reserved.
+
+#include "Text3DSceneViewExtension.h"
+#include "IText3DBuildSystem.h"
+#include "SceneView.h"
+#include "Utilities/Text3DUtilities.h"
+
+namespace UE::Text3D
+{
+
+void FText3DSceneViewExtension::SetupView(FSceneViewFamily& InViewFamily, FSceneView& InView)
+{
+	if (!InViewFamily.Scene)
+	{
+		return;
+	}
+	if (const IText3DBuildSystemInterface* const BuildSystem = Utilities::FindBuildSystem(InViewFamily.Scene->GetWorld()))
+	{
+		BuildSystem->HidePrimitivesBeingBuilt(InView);
+	}
+}
+
+bool FText3DSceneViewExtension::IsActiveThisFrame_Internal(const FSceneViewExtensionContext& InContext) const
+{
+	const IText3DBuildSystemInterface* const BuildSystem = Utilities::FindBuildSystem(InContext.GetWorld());
+	return BuildSystem && BuildSystem->IsPendingBuildCompletion();
+}
+
+} // UE::Text3D
diff --git a/Engine/Plugins/VirtualProduction/Text3D/Source/Text3D/Private/Text3DSceneViewExtension.h b/Engine/Plugins/VirtualProduction/Text3D/Source/Text3D/Private/Text3DSceneViewExtension.h
new file mode 100644
--- /dev/null
+++ b/Engine/Plugins/VirtualProduction/Text3D/Source/Text3D/Private/Text3DSceneViewExtension.h
@@ -0,0 +1,26 @@
+// Copyright Epic Games, Inc. All Rights Reserved.
+
+#pragma once
+
+#include "SceneViewExtension.h"
+
+namespace UE::Text3D
+{
+
+/** Scene view extensions to hide any primitive managed by Text3D that is still under construction */
+class FText3DSceneViewExtension : public FSceneViewExtensionBase
+{
+public:
+	explicit FText3DSceneViewExtension(const FAutoRegister& InAutoRegister)
+		: FSceneViewExtensionBase(InAutoRegister)
+	{
+	}
+
+protected:
+	//~ Begin FSceneViewExtensionBase
+	virtual void SetupView(FSceneViewFamily& InViewFamily, FSceneView& InView) override;
+	virtual bool IsActiveThisFrame_Internal(const FSceneViewExtensionContext& Context) const override;
+	//~ End FSceneViewExtensionBase
+};
+
+} // UE::Text3D
diff --git a/Engine/Plugins/VirtualProduction/Text3D/Source/Text3D/Public/IText3DBuildSystem.h b/Engine/Plugins/VirtualProduction/Text3D/Source/Text3D/Public/IText3DBuildSystem.h
--- a/Engine/Plugins/VirtualProduction/Text3D/Source/Text3D/Public/IText3DBuildSystem.h
+++ b/Engine/Plugins/VirtualProduction/Text3D/Source/Text3D/Public/IText3DBuildSystem.h
@@ -5,6 +5,8 @@
 #include "UObject/Interface.h"
 #include "IText3DBuildSystem.generated.h"
 
+class FSceneView;
+
 UINTERFACE(MinimalAPI, meta=(CannotImplementInterfaceInBlueprint))
 class UText3DBuildSystemInterface : public UInterface
 {
@@ -18,4 +20,7 @@ class IText3DBuildSystemInterface
 public:
 	/** Returns true if there are builders yet to complete */
 	virtual bool IsPendingBuildCompletion() const = 0;
+
+	/** Adds the Text3D-managed primitives that are being built to the scene view's hidden primitive list */
+	virtual void HidePrimitivesBeingBuilt(FSceneView& InSceneView) const = 0;
 };
diff --git a/Engine/Plugins/VirtualProduction/Text3D/Source/Text3D/Public/Renderers/Text3DRendererBase.h b/Engine/Plugins/VirtualProduction/Text3D/Source/Text3D/Public/Renderers/Text3DRendererBase.h
--- a/Engine/Plugins/VirtualProduction/Text3D/Source/Text3D/Public/Renderers/Text3DRendererBase.h
+++ b/Engine/Plugins/VirtualProduction/Text3D/Source/Text3D/Public/Renderers/Text3DRendererBase.h
@@ -8,6 +8,7 @@
 #include "Text3DRendererBase.generated.h"
 
 class UText3DComponent;
+class UPrimitiveComponent;
 
 /**
  * Base class for a rendering implementation of Text3D
@@ -34,6 +35,9 @@ public:
 	/** Get cached bounds from last update */
 	FBox GetBounds() const;
 
+	/** Calls the given functor for every primitive the renderer implementation manages */
+	void IterateManagedPrimitives(TFunctionRef<void(TNotNull<const UPrimitiveComponent*>)> InFunc) const;
+
 	/** Get the mesh type that this renderer uses */
 	virtual EText3DMeshType GetMeshType() const PURE_VIRTUAL(UText3DRendererBase::GetMeshType, return EText3DMeshType::Unknown; );
 
@@ -62,6 +66,11 @@ protected:
 	/** Calculate the bounds of the rendered text */
 	virtual FBox OnCalculateBounds() const PURE_VIRTUAL(UText3DRendererBase::OnCalculateBounds, return FBox(ForceInitToZero); );
 
+	/** Calls the given functor for every primitive the renderer implementation manages */
+	virtual void OnIterateManagedPrimitives(TFunctionRef<void(TNotNull<const UPrimitiveComponent*>)> InFunc) const
+	{
+	}
+
 #if WITH_EDITOR
 	void BindDebugDelegate();
 	void UnbindDebugDelegate();
diff --git a/Engine/Plugins/VirtualProduction/Text3D/Source/Text3D/Public/Text3DComponent.h b/Engine/Plugins/VirtualProduction/Text3D/Source/Text3D/Public/Text3DComponent.h
--- a/Engine/Plugins/VirtualProduction/Text3D/Source/Text3D/Public/Text3DComponent.h
+++ b/Engine/Plugins/VirtualProduction/Text3D/Source/Text3D/Public/Text3DComponent.h
@@ -552,6 +552,9 @@ private:
 	UPROPERTY(BlueprintAssignable, Category=Events, meta=(AllowPrivateAccess="true", DisplayName="On Text Generated"))
 	FTextGenerated TextGeneratedDelegate;
 
+	/** Accumulated update flags to trigger on OnRegister() */
+	EText3DRendererFlags PreRegisterUpdateFlags = EText3DRendererFlags::None;
+
 	FTextGeneratedNative TextGeneratedNativeDelegate;
 
 	/** Called before text is updated */
