Text3D: moved starved frames to be managed by each builder. Cvar to enable hiding primitives while building added (disabled by default).

#Synced-CL 51290016

diff --git a/Engine/Plugins/VirtualProduction/Text3D/Source/Text3D/Private/Subsystems/Text3DBuildSystem.cpp b/Engine/Plugins/VirtualProduction/Text3D/Source/Text3D/Private/Subsystems/Text3DBuildSystem.cpp
--- a/Engine/Plugins/VirtualProduction/Text3D/Source/Text3D/Private/Subsystems/Text3DBuildSystem.cpp
+++ b/Engine/Plugins/VirtualProduction/Text3D/Source/Text3D/Private/Subsystems/Text3DBuildSystem.cpp
@@ -24,7 +24,7 @@ TAutoConsoleVariable<int32> CVarFrameTimeBudgetMode(TEXT("Text3D.Build.FrameTime
 	);
 
 TAutoConsoleVariable<float> CVarFrameTimeBudget(TEXT("Text3D.Build.FrameTimeBudget")
-	, 0.016f
+	, 0.014f
 	, TEXT("The frame time cap to build Text3D components.\n")
 	  TEXT("Only relevant if Text3DBuild.FrameTimeBudgetMode is non-zero.\n")
 	  TEXT("If the current frame duration prior to the incremental build is past this limit, the build will be skipped (unless it has starved for more than the starve limit).\n")
@@ -101,10 +101,10 @@ void UText3DBuildSystem::RequestTextBuild(TNotNull<UText3DComponent*> InComponen
 		return;
 	}
 
-	UE_LOG(LogText3D, Verbose, TEXT("[%llu] Requesting %s text build for component '%s'. Flags: %d")
+	UE_LOG(LogText3D, Verbose, TEXT("[%llu] Requesting %s text build for component '%s'. Flags %0.2d")
 		, GFrameCounter
 		, bInBlockingBuild ? TEXT("BLOCKING") : TEXT("incremental")
-		, *InComponent->GetFullName()
+		, *InComponent->GetReadableName()
 		, static_cast<int32>(InFlags));
 
 	// allow if build is incremental. If blocking ensure the flags are exactly equal. If there's another flag, it can't be used.
@@ -118,10 +118,10 @@ void UText3DBuildSystem::RequestTextBuild(TNotNull<UText3DComponent*> InComponen
 	UE::Text3D::FTextBuilder* Builder = nullptr;
 	for (UE::Text3D::FTextBuilder& CandidateBuilder : TextBuilders)
 	{
-		if (CandidateBuilder.GetComponent() == InComponent && EnumHasAnyFlags(CandidateBuilder.GetFlags(), InFlags))
+		if (CandidateBuilder.GetComponent() == InComponent)
 		{
 			// Cancel builder if already in progress. A newer update is available.
-			if (CandidateBuilder.IsInProgress())
+			if (EnumHasAllFlags(InFlags, CandidateBuilder.GetFlags()) && CandidateBuilder.IsInProgress())
 			{
 				UE_LOG(LogText3D, Verbose, TEXT("\tCanceling active builder '%s'. A newer build request was made."), *CandidateBuilder.GetDebugString());
 				CandidateBuilder.CancelBuild();
@@ -187,7 +187,7 @@ int32 UText3DBuildSystem::RemoveTextBuild(TNotNull<UText3DComponent*> InComponen
 	RemovedCount += QueuedBuilders.RemoveAll(IsMatchingComponent);
 	RemovedCount += TextBuilders.RemoveAll(IsMatchingComponent);
 
-	UE_CLOG(RemovedCount > 0, LogText3D, Verbose, TEXT("Removed '%s' from %d builders"), *InComponent->GetFullName(), RemovedCount);
+	UE_CLOG(RemovedCount > 0, LogText3D, Verbose, TEXT("Removed '%s' from %d builders"), *InComponent->GetReadableName(), RemovedCount);
 	return RemovedCount;
 }
 
@@ -209,22 +209,22 @@ void UText3DBuildSystem::IncrementalTextBuild()
 		BeginFrameCycles.Reset();
 	};
 
-	UE_LOG(LogText3D, Verbose, TEXT("[%llu] Performing incremental Text3D build..."), GFrameCounter);
+	UE_LOG(LogText3D, VeryVerbose, TEXT("[%llu] Performing incremental Text3D build..."), GFrameCounter);
 
 	TextBuilders.Append(MoveTemp(QueuedBuilders));
 	QueuedBuilders.Reset();
 
 	// Remove all finished builders
-	const int32 RemovedCount = TextBuilders.RemoveAll(
+	TextBuilders.RemoveAll(
 		[](const UE::Text3D::FTextBuilder& InBuilder)
 		{
-			UE_CLOG(InBuilder.IsComplete(), LogText3D, Verbose, TEXT("\tRemoving finished Text3D builders '%s'"), *InBuilder.GetDebugString());
+			UE_CLOG(InBuilder.IsComplete(), LogText3D, Verbose, TEXT("\tRemoving finished Text3D builder '%s'"), *InBuilder.GetDebugString());
 			return InBuilder.IsComplete();
 		});
 
 	if (TextBuilders.IsEmpty())
 	{
-		UE_LOG(LogText3D, Verbose, TEXT("\tIncremental Text3D build finished. No more build jobs left."));
+		UE_LOG(LogText3D, VeryVerbose, TEXT("\tIncremental Text3D build finished. No more build jobs left."));
 		FCoreDelegates::OnEndFrame.Remove(EndFrameHandle);
 		EndFrameHandle.Reset();
 		return; // nothing to update
@@ -232,47 +232,39 @@ void UText3DBuildSystem::IncrementalTextBuild()
 
 	const uint64 CurrentCycles = FPlatformTime::Cycles64();
 
-	// If the incremental build has starved for too long, force it to happen
-	const bool bForceBuild = FramesStarved >= Private::CVarStarveFrameLimit.GetValueOnGameThread();
-	UE_CLOG(bForceBuild, LogText3D, Verbose, TEXT("\tIncremental Text3D build starved for more than %d frames. Forcing incremental build..."), Private::CVarStarveFrameLimit.GetValueOnGameThread());
-
-	TOptional<uint64> BudgetEnd;
+	UE::Text3D::FTextBuilder::FIncrementalBuildParams Params;
 	if (UE::Text3D::Private::CVarEnableBuildTimeBudget.GetValueOnGameThread())
 	{
-		BudgetEnd = CurrentCycles + FPlatformTime::SecondsToCycles64(UE::Text3D::Private::CVarBuildTimeBudget.GetValueOnGameThread());
-		// Only consider the budget deficit if not forcing a build
-		if (!bForceBuild)
+		Params.BudgetEnd = CurrentCycles + FPlatformTime::SecondsToCycles64(UE::Text3D::Private::CVarBuildTimeBudget.GetValueOnGameThread());
+		if (*Params.BudgetEnd <= BudgetDeficit)
 		{
-			if (*BudgetEnd <= BudgetDeficit)
-			{
-				++FramesStarved;
-				BudgetDeficit = 0;
-				UE_LOG(LogText3D, Verbose, TEXT("\tIncremental Text3D build ran out of build time budget due to blocking builds."));
-				return;
-			}
-			*BudgetEnd -= BudgetDeficit;
+			*Params.BudgetEnd = 0;
+		}
+		else
+		{
+			*Params.BudgetEnd -= BudgetDeficit;
 		}
 	}
 	BudgetDeficit = 0;
 
-	// Skip if the frame duration has already gone for too long
-	if (!bForceBuild && Private::IsFrameTimeBudgetEnabled(GetWorld()) && (CurrentCycles - *BeginFrameCycles) >= FPlatformTime::SecondsToCycles64(Private::CVarFrameTimeBudget.GetValueOnGameThread()))
+	Params.MaxStarvedFrames = Private::CVarStarveFrameLimit.GetValueOnGameThread();
+	Params.bFrameTimeExceeded = Private::IsFrameTimeBudgetEnabled(GetWorld()) 
+		&& (CurrentCycles - *BeginFrameCycles) >= FPlatformTime::SecondsToCycles64(Private::CVarFrameTimeBudget.GetValueOnGameThread());
+
+	Params.TempMaxUpdatePhase = EBuildPhase::PrepareExtensions;
+
+	// Update all the builders that are preparing
+	for (UE::Text3D::FTextBuilder& Builder : TextBuilders)
 	{
-		++FramesStarved;
-		UE_LOG(LogText3D, Verbose, TEXT("\tIncremental Text3D build ran out of frame time budget."));
-		return;
+		Builder.IncrementalBuild(Params);
 	}
-	FramesStarved = 0;
+
+	Params.TempMaxUpdatePhase.Reset();
 
 	// Update all the builders with no phase limit
 	for (UE::Text3D::FTextBuilder& Builder : TextBuilders)
 	{
-		const EBuildResult Result = Builder.IncrementalBuild(BudgetEnd);
-		if (Result == EBuildResult::OutOfBudget)
-		{
-			UE_LOG(LogText3D, Verbose, TEXT("\t\tIncremental Text3D build ran out of budget while building '%s'. Continuing next frame"), *Builder.GetDebugString());
-			return; // ran out of budget
-		}
+		Builder.IncrementalBuild(Params);
 	}
 }
 
@@ -354,7 +346,6 @@ void UText3DBuildSystem::ConditionallyRegisterDelegates()
 {
 	if (!BeginFrameHandle.IsValid())
 	{
-		FramesStarved = 0;
 		BeginFrameHandle = FCoreDelegates::OnBeginFrame.AddWeakLambda(this, [this]
 			{
 				BeginFrameCycles = FPlatformTime::Cycles64();
diff --git a/Engine/Plugins/VirtualProduction/Text3D/Source/Text3D/Private/Subsystems/Text3DBuildSystem.h b/Engine/Plugins/VirtualProduction/Text3D/Source/Text3D/Private/Subsystems/Text3DBuildSystem.h
--- a/Engine/Plugins/VirtualProduction/Text3D/Source/Text3D/Private/Subsystems/Text3DBuildSystem.h
+++ b/Engine/Plugins/VirtualProduction/Text3D/Source/Text3D/Private/Subsystems/Text3DBuildSystem.h
@@ -79,9 +79,6 @@ private:
 	/** Accumulated cycles spent on blocking updates to determine if next tick an update should happen */
 	uint64 BudgetDeficit = 0;
 
-	/** Number of frames the incremental build has starved for */
-	int32 FramesStarved = 0;
-
 	/** Handle to the begin frame call to keep track how long this frame's been going for */
 	FDelegateHandle BeginFrameHandle;
 
diff --git a/Engine/Plugins/VirtualProduction/Text3D/Source/Text3D/Private/Text3DBuilder.cpp b/Engine/Plugins/VirtualProduction/Text3D/Source/Text3D/Private/Text3DBuilder.cpp
--- a/Engine/Plugins/VirtualProduction/Text3D/Source/Text3D/Private/Text3DBuilder.cpp
+++ b/Engine/Plugins/VirtualProduction/Text3D/Source/Text3D/Private/Text3DBuilder.cpp
@@ -47,13 +47,13 @@ bool FTextBuilder::PrepareFlags(EText3DRendererFlags InFlags)
 
 FString FTextBuilder::GetDebugString() const
 {
-	return FString::Printf(TEXT("Build Flags %0.2d | Phase %s | Component %s")
-		, static_cast<int32>(Flags)
+	return FString::Printf(TEXT("%s | %s | Flags %0.2d")
+		, Component ? *Component->GetReadableName() : TEXT("(none)")
 		, Private::GetPhaseString(Phase)
-		, *GetFullNameSafe(Component));
+		, static_cast<int32>(Flags));
 }
 
-bool FTextBuilder::	IsInProgress() const
+bool FTextBuilder::IsInProgress() const
 {
 	return Phase != EBuildPhase::NotStarted && !IsComplete();
 }
@@ -68,35 +68,56 @@ void FTextBuilder::BlockingBuild()
 {
 	TRACE_CPUPROFILER_EVENT_SCOPE(Text3D::BlockingBuild);
 
+	StarvedFrames = 0;
 	BudgetEnd.Reset();
 	TempPhaseLimit.Reset();
 	BuildInternal();
 }
 
-EBuildResult FTextBuilder::IncrementalBuild(TOptional<uint64> InBudgetEnd, TOptional<EBuildPhase> InTempMaxUpdatePhase)
+void FTextBuilder::IncrementalBuild(const FIncrementalBuildParams& InParams)
 {
 	if (IsComplete())
 	{
-		return EBuildResult::BuildCompleted;
+		return;
 	}
 
 	if (Flags == EText3DRendererFlags::None)
 	{
 		Finish();
-		return EBuildResult::BuildCompleted;
+		return;
 	}
 
 	TRACE_CPUPROFILER_EVENT_SCOPE(Text3D::IncrementalBuild);
-
-	BudgetEnd = InBudgetEnd;
-	TempPhaseLimit = InTempMaxUpdatePhase;
+	BudgetEnd = InParams.BudgetEnd;
+	TempPhaseLimit = InParams.TempMaxUpdatePhase;
 
 	if (!IsWithinPhaseLimit())
 	{
 		// No work needed if already past the phase limit
-		return EBuildResult::PhaseCompleted;
+		return;
+	}
+
+	if (InParams.bFrameTimeExceeded || InParams.bBuildTimeExceeded)
+	{
+		if (StarvedFrames++ < InParams.MaxStarvedFrames)
+		{
+			UE_LOG(LogText3D, VeryVerbose, TEXT("\tIncremental Text3D build '%s' is out of budget. Continuing next frame")
+				, *GetDebugString());
+			return;
+		}
+
+		UE_LOG(LogText3D, VeryVerbose, TEXT("\tIncremental Text3D build '%s' is out of budget but has starved for more than %d frames. Forcing build...")
+			, *GetDebugString()
+			, InParams.MaxStarvedFrames);
+	}
+
+	StarvedFrames = 0;
+
+	const EBuildResult Result = BuildInternal();
+	if (Result == EBuildResult::OutOfBudget)
+	{
+		InParams.bBuildTimeExceeded = true;
 	}
-	return BuildInternal();
 }
 
 void FTextBuilder::CancelBuild()
@@ -198,7 +219,7 @@ bool FTextBuilder::IsWithinBudget() const
 
 EBuildResult FTextBuilder::IterateExtensions(TFunctionRef<bool(FExtensionElement&)> InFunc)
 {
-	while (IsWithinBudget())
+	do
 	{
 		if (Extensions.IsValidIndex(ExtensionIndex))
 		{
@@ -213,6 +234,8 @@ EBuildResult FTextBuilder::IterateExtensions(TFunctionRef<bool(FExtensionElement
 			return EBuildResult::PhaseCompleted;
 		}
 	}
+	while (IsWithinBudget());
+
 	return EBuildResult::OutOfBudget;
 }
 
diff --git a/Engine/Plugins/VirtualProduction/Text3D/Source/Text3D/Private/Text3DBuilder.h b/Engine/Plugins/VirtualProduction/Text3D/Source/Text3D/Private/Text3DBuilder.h
--- a/Engine/Plugins/VirtualProduction/Text3D/Source/Text3D/Private/Text3DBuilder.h
+++ b/Engine/Plugins/VirtualProduction/Text3D/Source/Text3D/Private/Text3DBuilder.h
@@ -91,7 +91,20 @@ public:
 	void BlockingBuild();
 
 	/** Incrementally builds the component every call with a budget end and update phase */
-	EBuildResult IncrementalBuild(TOptional<uint64> InBudgetEnd, TOptional<UE::Text3D::EBuildPhase> InTempMaxUpdatePhase = {});
+	struct FIncrementalBuildParams
+	{
+		/** If set, the cycle that the build should stop when past it. If unset, no budget end (analogous to a blocking build) */
+		TOptional<uint64> BudgetEnd;
+		/** The maximum number of frames the build frame can starve. */
+		int32 MaxStarvedFrames = 0;
+		/** If set, temporary max build phase that is allowed by the incremental update. If unset, no limit. */
+		TOptional<UE::Text3D::EBuildPhase> TempMaxUpdatePhase;
+		/** Whether the frame time has been exceeded already. Any non-starving incremental build will be skipped */
+		bool bFrameTimeExceeded = false;
+		/** Whether the build time has been exceeded already. Updated by the incremental build calls when this happens. Any non-starving incremental build will be skipped */
+		mutable bool bBuildTimeExceeded = false;
+	};
+	void IncrementalBuild(const FIncrementalBuildParams& InParams);
 
 	/** Cancels the build */
 	void CancelBuild();
@@ -152,6 +165,9 @@ private:
 	/** Current extension being updated. Relevant only to certain phases */
 	int32 ExtensionIndex = INDEX_NONE;
 
+	/** Number of frames that the incremental build has not been able to execute */
+	int32 StarvedFrames = 0;
+
 	/** Renderer flags for the built */
 	EText3DRendererFlags Flags = EText3DRendererFlags::None;
 
diff --git a/Engine/Plugins/VirtualProduction/Text3D/Source/Text3D/Private/Text3DSceneViewExtension.cpp b/Engine/Plugins/VirtualProduction/Text3D/Source/Text3D/Private/Text3DSceneViewExtension.cpp
--- a/Engine/Plugins/VirtualProduction/Text3D/Source/Text3D/Private/Text3DSceneViewExtension.cpp
+++ b/Engine/Plugins/VirtualProduction/Text3D/Source/Text3D/Private/Text3DSceneViewExtension.cpp
@@ -8,6 +8,15 @@
 namespace UE::Text3D
 {
 
+namespace Private
+{
+
+TAutoConsoleVariable<bool> CVarHidePrimitivesWhileBuilding(TEXT("Text3D.Build.HidePrimitivesWhileBuilding")
+	, false
+	, TEXT("If true, primitives managed by Text3D components that are being built will be hidden until completion."));
+
+} // UE::Text3D::Private
+
 void FText3DSceneViewExtension::SetupView(FSceneViewFamily& InViewFamily, FSceneView& InView)
 {
 	if (!InViewFamily.Scene)
@@ -22,6 +31,12 @@ void FText3DSceneViewExtension::SetupView(FSceneViewFamily& InViewFamily, FScene
 
 bool FText3DSceneViewExtension::IsActiveThisFrame_Internal(const FSceneViewExtensionContext& InContext) const
 {
+	const bool bEnabled = Private::CVarHidePrimitivesWhileBuilding.GetValueOnGameThread();
+	if (!bEnabled)
+	{
+		return false;
+	}
+
 	const IText3DBuildSystemInterface* const BuildSystem = Utilities::FindBuildSystem(InContext.GetWorld());
 	return BuildSystem && BuildSystem->IsPendingBuildCompletion();
 }
