Text3D: fixed issue with cache cleanup removing glyph mesh entries in between incremental build calls causing incremental build to set a null mesh.
+ Improved logging clarity

#Synced-CL 51290016

diff --git a/Engine/Plugins/VirtualProduction/Text3D/Source/Text3D/Private/Extensions/Text3DDefaultLayoutExtension.cpp b/Engine/Plugins/VirtualProduction/Text3D/Source/Text3D/Private/Extensions/Text3DDefaultLayoutExtension.cpp
--- a/Engine/Plugins/VirtualProduction/Text3D/Source/Text3D/Private/Extensions/Text3DDefaultLayoutExtension.cpp
+++ b/Engine/Plugins/VirtualProduction/Text3D/Source/Text3D/Private/Extensions/Text3DDefaultLayoutExtension.cpp
@@ -215,6 +215,8 @@ void UText3DDefaultLayoutExtension::PrepareBuild(EText3DRendererFlags InUpdateFl
 		return;
 	}
 
+	GlyphHandles.Empty();
+
 	TRACE_CPUPROFILER_EVENT_SCOPE(UText3DDefaultLayoutExtension::PrepareRebuild);
 
 	using namespace UE::Text3D::Layout;
@@ -273,6 +275,7 @@ void UText3DDefaultLayoutExtension::PrepareBuild(EText3DRendererFlags InUpdateFl
 			{
 				if (TOptional<FText3DBuildGlyphMeshDesc> GlyphOutline = FontFaceCache->PrepareCachedMesh(GlyphEntry->Entry.GlyphIndex, GlyphMeshParameters))
 				{
+					GlyphHandles.Emplace(*FontFaceCache, GlyphEntry->Entry.GlyphIndex, GlyphMeshParameters);
 					BuildGlyphMeshDescs.Emplace(MoveTemp(*GlyphOutline));
 				}
 			}
@@ -290,7 +293,10 @@ EText3DExtensionResult UText3DDefaultLayoutExtension::PreRendererUpdate(const UE
 {
 	using namespace UE::Text3D::Layout;
 
-	TRACE_CPUPROFILER_EVENT_SCOPE(UText3DDefaultLayoutExtension::PreRendererUpdate);
+	if (InParameters.CurrentFlag != EText3DRendererFlags::Geometry && InParameters.CurrentFlag != EText3DRendererFlags::Layout)
+	{
+		return EText3DExtensionResult::Active;
+	}
 
 	if (!GlyphText.IsValid())
 	{
@@ -298,10 +304,7 @@ EText3DExtensionResult UText3DDefaultLayoutExtension::PreRendererUpdate(const UE
 		return EText3DExtensionResult::Failed;
 	}
 
-	if (InParameters.CurrentFlag != EText3DRendererFlags::Geometry && InParameters.CurrentFlag != EText3DRendererFlags::Layout)
-	{
-		return EText3DExtensionResult::Active;
-	}
+	TRACE_CPUPROFILER_EVENT_SCOPE(UText3DDefaultLayoutExtension::PreRendererUpdate);
 
 	const UText3DComponent* Text3DComponent = GetText3DComponent();
 
@@ -476,6 +479,11 @@ EText3DExtensionResult UText3DDefaultLayoutExtension::PreRendererUpdate(const UE
 		}
 	}
 
+	if (EnumHasAnyFlags(InParameters.UpdateFlags, EText3DRendererFlags::Geometry))
+	{
+		// Glyph meshes have been processed. No longer need to keep handle to these.
+		GlyphHandles.Empty();
+	}
 	return EText3DExtensionResult::Finished;
 #else
 	UE_LOG(LogText3D, Warning, TEXT("FreeType is not available, cannot proceed without it"))
diff --git a/Engine/Plugins/VirtualProduction/Text3D/Source/Text3D/Private/Subsystems/Text3DBuildSystem.cpp b/Engine/Plugins/VirtualProduction/Text3D/Source/Text3D/Private/Subsystems/Text3DBuildSystem.cpp
--- a/Engine/Plugins/VirtualProduction/Text3D/Source/Text3D/Private/Subsystems/Text3DBuildSystem.cpp
+++ b/Engine/Plugins/VirtualProduction/Text3D/Source/Text3D/Private/Subsystems/Text3DBuildSystem.cpp
@@ -4,6 +4,7 @@
 #include "Algo/Contains.h"
 #include "Algo/ForEach.h"
 #include "Components/PrimitiveComponent.h"
+#include "Engine/Level.h"
 #include "Engine/World.h"
 #include "Logs/Text3DLogs.h"
 #include "Misc/CoreDelegates.h"
@@ -12,6 +13,7 @@
 #include "SceneViewExtension.h"
 #include "Text3DComponent.h"
 #include "Text3DSceneViewExtension.h"
+#include "Utilities/Text3DScopedLog.h"
 
 namespace UE::Text3D::Private
 {
@@ -101,14 +103,22 @@ void UText3DBuildSystem::RequestTextBuild(TNotNull<UText3DComponent*> InComponen
 		return;
 	}
 
-	UE_LOG(LogText3D, Verbose, TEXT("[%llu] Requesting %s text build for component '%s'. Flags %0.2d")
-		, GFrameCounter
+	UE::Text3D::FScopedLog ScopedLog;
+	UE_LOG(LogText3D, VeryVerbose, TEXT("%sRequesting %s text build for component '%s'. Flags %0.2d")
+		, ScopedLog.GetLogPrefix()
 		, bInBlockingBuild ? TEXT("BLOCKING") : TEXT("incremental")
 		, *InComponent->GetReadableName()
 		, static_cast<int32>(InFlags));
 
-	// allow if build is incremental. If blocking ensure the flags are exactly equal. If there's another flag, it can't be used.
-	auto IsBuilderFlagsCompatible = 
+	// A builder is relevant if it matches the component and has not complete
+	auto IsBuilderRelevant = 
+		[InComponent](const UE::Text3D::FTextBuilder& InBuilder)->bool
+		{
+			return InBuilder.GetComponent() == InComponent && !InBuilder.IsComplete();
+		};
+
+	// Always allow if build is incremental. If blocking ensure the flags are exactly equal. If there's another flag, it can't be used.
+	auto AreBuilderFlagsCompatible = 
 		[InFlags, bInBlockingBuild](const UE::Text3D::FTextBuilder& InBuilder)->bool
 		{
 			return !bInBlockingBuild || InBuilder.GetFlags() == InFlags;
@@ -118,19 +128,26 @@ void UText3DBuildSystem::RequestTextBuild(TNotNull<UText3DComponent*> InComponen
 	UE::Text3D::FTextBuilder* Builder = nullptr;
 	for (UE::Text3D::FTextBuilder& CandidateBuilder : TextBuilders)
 	{
-		if (CandidateBuilder.GetComponent() == InComponent)
+		if (IsBuilderRelevant(CandidateBuilder))
 		{
 			// Cancel builder if already in progress. A newer update is available.
-			if (EnumHasAllFlags(InFlags, CandidateBuilder.GetFlags()) && CandidateBuilder.IsInProgress())
+			// Only applicable if the new flags contain all the builders flags
+			// E.g. we don't want to cancel a Geometry&Material build when requesting a Material only build
+			// But we do want to cancel a Layout build when requesting a Geometry&Layout build
+			if (EnumHasAllFlags(InFlags, CandidateBuilder.GetFlags()))
 			{
-				UE_LOG(LogText3D, Verbose, TEXT("\tCanceling active builder '%s'. A newer build request was made."), *CandidateBuilder.GetDebugString());
+				UE_LOG(LogText3D, VeryVerbose, TEXT("%sCanceling active builder '%s'. A newer build request was made.")
+					, UE::Text3D::FScopedLog().GetLogPrefix()
+					, *CandidateBuilder.GetDebugName());
 				CandidateBuilder.CancelBuild();
 				continue;
 			}
 
-			if (IsBuilderFlagsCompatible(CandidateBuilder) && EnumHasAllFlags(CandidateBuilder.GetFlags(), InFlags))
+			if (AreBuilderFlagsCompatible(CandidateBuilder) && EnumHasAllFlags(CandidateBuilder.GetFlags(), InFlags))
 			{
-				UE_LOG(LogText3D, Verbose, TEXT("\tFound candidate active builder '%s'."), *CandidateBuilder.GetDebugString());
+				UE_LOG(LogText3D, VeryVerbose, TEXT("%sFound candidate active builder '%s'.")
+					, UE::Text3D::FScopedLog().GetLogPrefix()	
+					, *CandidateBuilder.GetDebugName());
 				Builder = &CandidateBuilder;
 				break;
 			}
@@ -142,9 +159,11 @@ void UText3DBuildSystem::RequestTextBuild(TNotNull<UText3DComponent*> InComponen
 	{
 		for (UE::Text3D::FTextBuilder& CandidateBuilder : QueuedBuilders)
 		{
-			if (CandidateBuilder.GetComponent() == InComponent && IsBuilderFlagsCompatible(CandidateBuilder))
+			if (IsBuilderRelevant(CandidateBuilder) && AreBuilderFlagsCompatible(CandidateBuilder))
 			{
-				UE_LOG(LogText3D, Verbose, TEXT("\tFound candidate queued builder '%s'."), *CandidateBuilder.GetDebugString());
+				UE_LOG(LogText3D, VeryVerbose, TEXT("%sFound candidate queued builder '%s'.")
+					, UE::Text3D::FScopedLog().GetLogPrefix()
+					, *CandidateBuilder.GetDebugName());
 				Builder = &CandidateBuilder;
 				break;
 			}
@@ -159,7 +178,10 @@ void UText3DBuildSystem::RequestTextBuild(TNotNull<UText3DComponent*> InComponen
 	{
 		Builder = &QueuedBuilders.Emplace_GetRef(InComponent);
 		Builder->PrepareFlags(InFlags);
-		UE_LOG(LogText3D, Verbose, TEXT("\tNo candidate text builder was found. Queueing a new builder '%s'"), *Builder->GetDebugString());
+
+		UE_LOG(LogText3D, VeryVerbose, TEXT("%sNo candidate text builder was found. Queueing a new builder '%s'")
+			, UE::Text3D::FScopedLog().GetLogPrefix()
+			, *Builder->GetDebugName());
 	}
 
 	ConditionallyRegisterDelegates();
@@ -209,7 +231,8 @@ void UText3DBuildSystem::IncrementalTextBuild()
 		BeginFrameCycles.Reset();
 	};
 
-	UE_LOG(LogText3D, VeryVerbose, TEXT("[%llu] Performing incremental Text3D build..."), GFrameCounter);
+	UE::Text3D::FScopedLog ScopedLog;
+	UE_LOG(LogText3D, VeryVerbose, TEXT("%sPerforming incremental Text3D build..."), ScopedLog.GetLogPrefix());
 
 	TextBuilders.Append(MoveTemp(QueuedBuilders));
 	QueuedBuilders.Reset();
@@ -218,13 +241,15 @@ void UText3DBuildSystem::IncrementalTextBuild()
 	TextBuilders.RemoveAll(
 		[](const UE::Text3D::FTextBuilder& InBuilder)
 		{
-			UE_CLOG(InBuilder.IsComplete(), LogText3D, Verbose, TEXT("\tRemoving finished Text3D builder '%s'"), *InBuilder.GetDebugString());
+			UE_CLOG(InBuilder.IsComplete(), LogText3D, VeryVerbose, TEXT("%sRemoving finished builder '%s'")
+				, UE::Text3D::FScopedLog().GetLogPrefix()
+				, *InBuilder.GetDebugName());
 			return InBuilder.IsComplete();
 		});
 
-	if (TextBuilders.IsEmpty())
+	if (!IsBuildingText())
 	{
-		UE_LOG(LogText3D, VeryVerbose, TEXT("\tIncremental Text3D build finished. No more build jobs left."));
+		UE_LOG(LogText3D, VeryVerbose, TEXT("%sIncremental Text3D build finished. No more build jobs left."), UE::Text3D::FScopedLog().GetLogPrefix());
 		FCoreDelegates::OnEndFrame.Remove(EndFrameHandle);
 		EndFrameHandle.Reset();
 		return; // nothing to update
@@ -268,6 +293,26 @@ void UText3DBuildSystem::IncrementalTextBuild()
 	}
 }
 
+bool UText3DBuildSystem::IsBuildingText() const
+{
+	if (!QueuedBuilders.IsEmpty() || !TextBuilders.IsEmpty())
+	{
+		return true;
+	}
+
+	// Assume that if there are levels that are still not fully registered, that text is still pending build.
+	// There could be Text3D components that are yet to register to request builds.
+	const UWorld* const World = GetWorld();
+	for (const ULevel* Level : World->GetLevels())
+	{
+		if (Level && !Level->bAreComponentsCurrentlyRegistered)
+		{
+			return true;
+		}
+	}
+	return false;
+}
+
 void UText3DBuildSystem::Initialize(FSubsystemCollectionBase& InCollection)
 {
 	Super::Initialize(InCollection);
@@ -303,14 +348,9 @@ bool UText3DBuildSystem::DoesSupportWorldType(const EWorldType::Type InWorldType
 	return InWorldType != EWorldType::None;
 }
 
-bool UText3DBuildSystem::IsPendingBuildCompletion() const
+bool UText3DBuildSystem::IsBuildInProgress() const
 {
-	// Pending build if there are queued builders, or if there's at least one active builder that hasn't finished.
-	return !QueuedBuilders.IsEmpty() || TextBuilders.ContainsByPredicate(
-		[](const UE::Text3D::FTextBuilder& InTextBuilder)
-		{
-			return !InTextBuilder.IsComplete();
-		});
+	return IsBuildingText();
 }
 
 void UText3DBuildSystem::HidePrimitivesBeingBuilt(FSceneView& InSceneView) const
diff --git a/Engine/Plugins/VirtualProduction/Text3D/Source/Text3D/Private/Subsystems/Text3DBuildSystem.h b/Engine/Plugins/VirtualProduction/Text3D/Source/Text3D/Private/Subsystems/Text3DBuildSystem.h
--- a/Engine/Plugins/VirtualProduction/Text3D/Source/Text3D/Private/Subsystems/Text3DBuildSystem.h
+++ b/Engine/Plugins/VirtualProduction/Text3D/Source/Text3D/Private/Subsystems/Text3DBuildSystem.h
@@ -20,7 +20,7 @@ enum class EText3DRendererFlags : uint8;
  * Manages all the text builds for a world 
  * It's different from UText3DEngineSubsystem to keep these text builders referencing the text components scoped within the same world.
  */
-UCLASS()
+UCLASS(BlueprintType)
 class UText3DBuildSystem : public UWorldSubsystem, public IText3DBuildSystemInterface
 {
 	GENERATED_BODY()
@@ -42,6 +42,10 @@ public:
 	/** Called on tick to process builders and call IncrementalBuild on them */
 	void IncrementalTextBuild();
 
+	/** Returns true if there are Text3D components being built */
+	UFUNCTION(BlueprintPure, Category="Text 3D")
+	bool IsBuildingText() const;
+
 protected:
 	//~ Begin USubsystem
 	virtual void Initialize(FSubsystemCollectionBase& InCollection) override;
@@ -53,7 +57,7 @@ protected:
 	//~ End UWorldSubsystem
 
 	//~ Begin IText3DBuildSystemInterface
-	virtual bool IsPendingBuildCompletion() const override;
+	virtual bool IsBuildInProgress() const override;
 	virtual void HidePrimitivesBeingBuilt(FSceneView& InSceneView) const override;
 	//~ End IText3DBuildSystemInterface
 
diff --git a/Engine/Plugins/VirtualProduction/Text3D/Source/Text3D/Private/Subsystems/Text3DEngineSubsystem.cpp b/Engine/Plugins/VirtualProduction/Text3D/Source/Text3D/Private/Subsystems/Text3DEngineSubsystem.cpp
--- a/Engine/Plugins/VirtualProduction/Text3D/Source/Text3D/Private/Subsystems/Text3DEngineSubsystem.cpp
+++ b/Engine/Plugins/VirtualProduction/Text3D/Source/Text3D/Private/Subsystems/Text3DEngineSubsystem.cpp
@@ -76,7 +76,27 @@ bool UText3DEngineSubsystem::Exec_Dev(UWorld* InWorld, const TCHAR* InCmd, FOutp
 
 			if (FParse::Command(&InCmd, TEXT("cleanup")))
 			{
-				CleanupCache(0);
+				if (bGlyphMeshBuildInProgress)
+				{
+					UE_LOG(LogText3D, Warning, TEXT("Cleanup failed. Text3D Engine Subsystem is currently building glyph meshes."));
+				}
+				else
+				{
+					CleanupCache(0);	
+				}
+				return true;
+			}
+
+			if (FParse::Command(&InCmd, TEXT("clear")))
+			{
+				if (bGlyphMeshBuildInProgress)
+				{
+					UE_LOG(LogText3D, Warning, TEXT("Clear cache failed. Text3D Engine Subsystem is currently building glyph meshes."));
+				}
+				else
+				{
+					ClearCache();	
+				}
 				return true;
 			}
 		}
@@ -114,8 +134,10 @@ void UText3DEngineSubsystem::ScheduleCacheCleanup()
 	const UText3DProjectSettings* const TextSettings = UText3DProjectSettings::Get();
 	const float CleanupDelay = TextSettings ? TextSettings->GetFontFaceGlyphCleanupDelay() : 5.f;
 
+	UE_LOG(LogText3D, Verbose, TEXT("Text3D Engine Subsystem cache cleanup requested."));
+
 	// Immediate cleanup if close to zero or negative.
-	if (CleanupDelay <= UE_SMALL_NUMBER)
+	if (CleanupDelay <= UE_SMALL_NUMBER && !bGlyphMeshBuildInProgress)
 	{
 		CleanupCache(0.f);
 	}
@@ -136,9 +158,12 @@ bool UText3DEngineSubsystem::CleanupCache(float)
 	if (bGlyphMeshBuildInProgress)
 	{
 		// Wait until next tick to see if cache can be cleaned up, glyph mesh building in progress
+		UE_LOG(LogText3D, Verbose, TEXT("Text3D Engine Subsystem cache cleanup could not be performed while glyph mesh building in progress. Trying again next call."));
 		return true;
 	}
 
+	UE_LOG(LogText3D, Verbose, TEXT("Text3D Engine Subsystem performing cache cleanup..."));
+
 	// Removing CleanupHandle Ticker while still executing this ticker delegate is by-design valid.
 	// see TickerTests.cpp:112 "a delegate removal from inside the delegate execution (used to be a deadlock)"
 	FTSTicker::RemoveTicker(CacheCleanupDelegate);
@@ -174,7 +199,7 @@ void UText3DEngineSubsystem::PrintCache() const
 	UE_LOG(LogText3D, Log, TEXT("Text3D Engine Subsystem is currently caching %i fonts: "), CachedFontFaces.Num());
 	for (const TPair<uint32, FText3DFontFaceCache>& CachedFont : CachedFontFaces)
 	{
-		CachedFont.Value.PrintCache();
+		UE_LOG(LogText3D, Log, TEXT("\t%s"), *CachedFont.Value.ToDebugString());
 	}
 }
 
@@ -340,7 +365,6 @@ TOptional<FText3DBuildGlyphMeshDesc> FText3DFontFaceCache::PrepareCachedMesh(uin
 	MeshDesc.GlyphOutline = MoveTemp(*GlyphOutline);
 	MeshDesc.FontFace = FontFace;
 	MeshDesc.GlyphIndex = InGlyphIndex;
-	MeshDesc.GlyphHandle = UE::Text3D::Geometry::FCachedFontFaceGlyphHandle(*this, InGlyphIndex, InParameters);
 	return MeshDesc;
 }
 
@@ -479,14 +503,24 @@ TOptional<UE::Text3D::FGlyphOutline> FText3DFontFaceCache::LoadGlyphOutline(uint
 #endif
 }
 
-void FText3DFontFaceCache::PrintCache() const
+FString FText3DFontFaceCache::ToDebugString() const
 {
 #if WITH_FREETYPE
-	if (IsValid())
+	if (!IsValid())
 	{
-		const FT_Face Face = FontFace->GetFace();
-		UE_LOG(LogText3D, Log, TEXT("Cached Font Face Family=%s Style=%s Index=%i Hash=%u Meshes=%i Usage=%i"), *FString(Face->family_name), *FString(Face->style_name), Face->face_index, FontFaceHash, GlyphMeshes.Num(), GetCacheRefCount());
+		return TEXT("(invalid font face cache)");
 	}
+
+	const FT_Face Face = FontFace->GetFace();
+	return FString::Printf(TEXT("Cached Font Face Family=%s Style=%s Index=%i Hash=%u Meshes=%i Usage=%i")
+		, Face->family_name ? *FString(Face->family_name) : TEXT("(invalid)")
+		, Face->style_name ? *FString(Face->style_name) : TEXT("(invalid)")
+		, Face->face_index
+		, FontFaceHash
+		, GlyphMeshes.Num()
+		, GetCacheRefCount());
+#else
+	return FString();
 #endif
 }
 
diff --git a/Engine/Plugins/VirtualProduction/Text3D/Source/Text3D/Private/Subsystems/Text3DEngineSubsystem.h b/Engine/Plugins/VirtualProduction/Text3D/Source/Text3D/Private/Subsystems/Text3DEngineSubsystem.h
--- a/Engine/Plugins/VirtualProduction/Text3D/Source/Text3D/Private/Subsystems/Text3DEngineSubsystem.h
+++ b/Engine/Plugins/VirtualProduction/Text3D/Source/Text3D/Private/Subsystems/Text3DEngineSubsystem.h
@@ -22,8 +22,6 @@ struct FText3DBuildGlyphMeshDesc
 	TSharedPtr<FFreeTypeFace> FontFace;
 	/** Index to the glyph to build */
 	uint32 GlyphIndex = 0;
-	/** Holds ref count of the glyph to prevent destruction */
-	UE::Text3D::Geometry::FCachedFontFaceGlyphHandle GlyphHandle;
 	/** Set when queuing the mesh desc to map to the glyph parameters used */
 	int32 GlyphParametersIndex = 0;
 };
@@ -74,7 +72,7 @@ struct FText3DFontFaceCache
 	FText3DCachedMesh* FindGlyphMesh(uint32 InGlyphHash);
 	uint32 CleanupCache();
 	
-	void PrintCache() const;
+	FString ToDebugString() const;
 	bool IsValid() const;
 	uint32 GetCacheRefCount() const;
 	uint32 GetCacheGlyphCount() const;
diff --git a/Engine/Plugins/VirtualProduction/Text3D/Source/Text3D/Private/Text3DBuilder.cpp b/Engine/Plugins/VirtualProduction/Text3D/Source/Text3D/Private/Text3DBuilder.cpp
--- a/Engine/Plugins/VirtualProduction/Text3D/Source/Text3D/Private/Text3DBuilder.cpp
+++ b/Engine/Plugins/VirtualProduction/Text3D/Source/Text3D/Private/Text3DBuilder.cpp
@@ -5,6 +5,7 @@
 #include "Logs/Text3DLogs.h"
 #include "Renderers/Text3DRendererBase.h"
 #include "Text3DComponent.h"
+#include "Utilities/Text3DScopedLog.h"
 
 namespace UE::Text3D
 {
@@ -20,7 +21,7 @@ const TCHAR* GetPhaseString(EBuildPhase InPhase)
 	case EBuildPhase::PrepareExtensions:  return TEXT("PrepareExtensions");
 	case EBuildPhase::PreRendererUpdate:  return TEXT("PreRendererUpdate");
 	case EBuildPhase::RendererUpdate:     return TEXT("RendererUpdate");
-	case EBuildPhase::PostRendererUpdate: return TEXT("PrepareExtensions");
+	case EBuildPhase::PostRendererUpdate: return TEXT("PostRendererUpdate");
 	case EBuildPhase::Succeeded:          return TEXT("Succeeded");
 	case EBuildPhase::Canceled:           return TEXT("Canceled");
 	}
@@ -45,10 +46,23 @@ bool FTextBuilder::PrepareFlags(EText3DRendererFlags InFlags)
 	return true;
 }
 
-FString FTextBuilder::GetDebugString() const
+FString FTextBuilder::GetDebugName() const
 {
-	return FString::Printf(TEXT("%s | %s | Flags %0.2d")
+	return FString::Printf(TEXT("%s (Status: %s)")
 		, Component ? *Component->GetReadableName() : TEXT("(none)")
+		, *GetDebugStatus());
+}
+
+FString FTextBuilder::GetDebugStatus() const
+{
+	if (IsInRenderingPhases())
+	{
+		return FString::Printf(TEXT("%s | Flags %0.2d | Current Flag %0.2d")
+			, Private::GetPhaseString(Phase)
+			, static_cast<int32>(Flags)
+			, static_cast<int32>(CurrentFlag));
+	}
+	return FString::Printf(TEXT("%s | Flags %0.2d")
 		, Private::GetPhaseString(Phase)
 		, static_cast<int32>(Flags));
 }
@@ -68,6 +82,9 @@ void FTextBuilder::BlockingBuild()
 {
 	TRACE_CPUPROFILER_EVENT_SCOPE(Text3D::BlockingBuild);
 
+	FScopedLog ScopedLog;
+	UE_LOG(LogText3D, VeryVerbose, TEXT("%sPerforming BLOCKING build for '%s'..."), ScopedLog.GetLogPrefix(), *GetDebugName());
+
 	StarvedFrames = 0;
 	BudgetEnd.Reset();
 	TempPhaseLimit.Reset();
@@ -101,21 +118,22 @@ void FTextBuilder::IncrementalBuild(const FIncrementalBuildParams& InParams)
 	{
 		if (StarvedFrames++ < InParams.MaxStarvedFrames)
 		{
-			UE_LOG(LogText3D, VeryVerbose, TEXT("\tIncremental Text3D build '%s' is out of budget. Continuing next frame")
-				, *GetDebugString());
 			return;
 		}
 
-		UE_LOG(LogText3D, VeryVerbose, TEXT("\tIncremental Text3D build '%s' is out of budget but has starved for more than %d frames. Forcing build...")
-			, *GetDebugString()
+		UE_LOG(LogText3D, VeryVerbose, TEXT("%sIncremental build is out of budget but has starved for more than %d frames. Forcing build...")
+			, FScopedLog().GetLogPrefix()
 			, InParams.MaxStarvedFrames);
 	}
 
+	FScopedLog ScopedLog;
+	UE_LOG(LogText3D, VeryVerbose, TEXT("%sPerforming incremental build for '%s'..."), ScopedLog.GetLogPrefix(), *GetDebugName());
 	StarvedFrames = 0;
 
 	const EBuildResult Result = BuildInternal();
 	if (Result == EBuildResult::OutOfBudget)
 	{
+		UE_LOG(LogText3D, VeryVerbose, TEXT("%sBuild ran out of budget. Continuing next frame"), FScopedLog().GetLogPrefix());
 		InParams.bBuildTimeExceeded = true;
 	}
 }
@@ -155,8 +173,16 @@ void FTextBuilder::AddReferencedObjects(FReferenceCollector& InCollector)
 	}
 }
 
+bool FTextBuilder::IsInRenderingPhases() const
+{
+	return Phase >= EBuildPhase::PreRendererUpdate 
+		&& Phase <= EBuildPhase::PostRendererUpdate;
+}
+
 EBuildResult FTextBuilder::BuildInternal()
 {
+	UE_LOG(LogText3D, VeryVerbose, TEXT("%sPerforming Build. Current Status: %s"), FScopedLog().GetLogPrefix(), *GetDebugStatus());	
+
 	EBuildResult Result;
 
 	switch (Phase)
@@ -418,6 +444,8 @@ void FTextBuilder::Finish()
 	Extensions.Reset();
 	ExtensionIndex = INDEX_NONE;
 	Phase = EBuildPhase::Succeeded;
+
+	UE_LOG(LogText3D, VeryVerbose, TEXT("%sText3D Builder Finished."), FScopedLog().GetLogPrefix());
 }
 
 } // UE::Text3D
diff --git a/Engine/Plugins/VirtualProduction/Text3D/Source/Text3D/Private/Text3DBuilder.h b/Engine/Plugins/VirtualProduction/Text3D/Source/Text3D/Private/Text3DBuilder.h
--- a/Engine/Plugins/VirtualProduction/Text3D/Source/Text3D/Private/Text3DBuilder.h
+++ b/Engine/Plugins/VirtualProduction/Text3D/Source/Text3D/Private/Text3DBuilder.h
@@ -79,7 +79,10 @@ public:
 	}
 
 	/** Gets a debug string to identify the builder */
-	FString GetDebugString() const;
+	FString GetDebugName() const;
+
+	/** Gets a debug string to identify the builder */
+	FString GetDebugStatus() const;
 
 	/** Returns true if the builder is currently in progress (i.e. between not-started and completed) */
 	bool IsInProgress() const;
@@ -116,6 +119,9 @@ public:
 	void AddReferencedObjects(FReferenceCollector& InCollector);
 
 private:
+	/** Whether the builder is currently in rendering phases */
+	bool IsInRenderingPhases() const;
+
 	/** Calls the appropriate function depending on the current phase */
 	EBuildResult BuildInternal();
 
diff --git a/Engine/Plugins/VirtualProduction/Text3D/Source/Text3D/Private/Text3DSceneViewExtension.cpp b/Engine/Plugins/VirtualProduction/Text3D/Source/Text3D/Private/Text3DSceneViewExtension.cpp
--- a/Engine/Plugins/VirtualProduction/Text3D/Source/Text3D/Private/Text3DSceneViewExtension.cpp
+++ b/Engine/Plugins/VirtualProduction/Text3D/Source/Text3D/Private/Text3DSceneViewExtension.cpp
@@ -38,7 +38,7 @@ bool FText3DSceneViewExtension::IsActiveThisFrame_Internal(const FSceneViewExten
 	}
 
 	const IText3DBuildSystemInterface* const BuildSystem = Utilities::FindBuildSystem(InContext.GetWorld());
-	return BuildSystem && BuildSystem->IsPendingBuildCompletion();
+	return BuildSystem && BuildSystem->IsBuildInProgress();
 }
 
 } // UE::Text3D
diff --git a/Engine/Plugins/VirtualProduction/Text3D/Source/Text3D/Private/Text3DTypes.cpp b/Engine/Plugins/VirtualProduction/Text3D/Source/Text3D/Private/Text3DTypes.cpp
--- a/Engine/Plugins/VirtualProduction/Text3D/Source/Text3D/Private/Text3DTypes.cpp
+++ b/Engine/Plugins/VirtualProduction/Text3D/Source/Text3D/Private/Text3DTypes.cpp
@@ -1,7 +1,7 @@
 // Copyright Epic Games, Inc. All Rights Reserved.
 
 #include "Text3DTypes.h"
-
+#include "Logs/Text3DLogs.h"
 #include "Subsystems/Text3DEngineSubsystem.h"
 
 namespace UE::Text3D::Geometry
@@ -11,10 +11,19 @@ namespace UE::Text3D::Geometry
 		GlyphIndex = InGlyphIndex;
 		FontFaceHash = GetTypeHash(InFontFaceCache);
 		GlyphMeshHash = FText3DFontFaceCache::GetGlyphMeshHash(InGlyphIndex, InGlyphMeshParameters);
+
 		if (FText3DCachedMesh* CachedMesh = Resolve())
 		{
 			CachedMesh->RefCount++;
 		}
+		else
+		{
+			UE_LOG(LogText3D, Warning, TEXT("Constructed glyph handle does not point to a valid entry! No ref counting will occur. Glyph Index: %d, Font Face Hash: %u, Glyph Mesh Hash: %u. Font Face: %s")
+				, GlyphIndex
+				, FontFaceHash
+				, GlyphMeshHash
+				, *InFontFaceCache.ToDebugString());
+		}
 	}
 
 	FCachedFontFaceGlyphHandle::FCachedFontFaceGlyphHandle(const FCachedFontFaceGlyphHandle& InOther)
diff --git a/Engine/Plugins/VirtualProduction/Text3D/Source/Text3D/Private/Utilities/Text3DScopedLog.cpp b/Engine/Plugins/VirtualProduction/Text3D/Source/Text3D/Private/Utilities/Text3DScopedLog.cpp
new file mode 100644
--- /dev/null
+++ b/Engine/Plugins/VirtualProduction/Text3D/Source/Text3D/Private/Utilities/Text3DScopedLog.cpp
@@ -0,0 +1,48 @@
+// Copyright Epic Games, Inc. All Rights Reserved.
+
+#include "Text3DScopedLog.h"
+
+#if UE_TEXT3D_LOG_SCOPE_ENABLED
+#include "Containers/UnrealString.h"
+#include "CoreGlobals.h"
+#include "Templates/Tuple.h"
+
+namespace UE::Text3D
+{
+
+namespace Private
+{
+
+static int32 ScopeLevel = -1;
+
+} // UE::Text3D::Private
+
+FScopedLog::FScopedLog()
+{
+	++Private::ScopeLevel;
+}
+
+FScopedLog::~FScopedLog()
+{
+	--Private::ScopeLevel;
+}
+
+const TCHAR* FScopedLog::GetLogPrefix() const
+{
+	if (Private::ScopeLevel <= 0)
+	{
+		static TPair<uint64, FString> CachedFrameCounter(0, FString());
+		if (CachedFrameCounter.Key != GFrameCounter)
+		{
+			CachedFrameCounter.Value = FString::Printf(TEXT("[%llu] "), GFrameCounter);
+		}
+		return *CachedFrameCounter.Value;
+	}
+	constexpr int32 NumSpacesPerScopeLevel = 3;
+	// FCString::Spc returns a ptr to a static string of spaces
+	return FCString::Spc(Private::ScopeLevel * NumSpacesPerScopeLevel);
+}
+
+} // UE::Text3D
+
+#endif // UE_TEXT3D_LOG_SCOPE_ENABLED
diff --git a/Engine/Plugins/VirtualProduction/Text3D/Source/Text3D/Private/Utilities/Text3DScopedLog.h b/Engine/Plugins/VirtualProduction/Text3D/Source/Text3D/Private/Utilities/Text3DScopedLog.h
new file mode 100644
--- /dev/null
+++ b/Engine/Plugins/VirtualProduction/Text3D/Source/Text3D/Private/Utilities/Text3DScopedLog.h
@@ -0,0 +1,32 @@
+// Copyright Epic Games, Inc. All Rights Reserved.
+
+#pragma once
+
+#include "Misc/Build.h"
+#include "HAL/Platform.h"
+
+#define UE_TEXT3D_LOG_SCOPE_ENABLED 1 && !NO_LOGGING
+
+namespace UE::Text3D
+{
+
+#if UE_TEXT3D_LOG_SCOPE_ENABLED
+/** Handles adding and removing 'scope levels' for log clarity */
+struct FScopedLog
+{
+	FScopedLog();
+	~FScopedLog();
+
+	const TCHAR* GetLogPrefix() const;
+};
+#else
+struct FScopedLog
+{
+	const TCHAR* GetLogPrefix() const
+	{
+		return TEXT("");
+	}
+};
+#endif // UE_TEXT3D_LOG_SCOPE_ENABLED
+
+} // UE::Text3D
diff --git a/Engine/Plugins/VirtualProduction/Text3D/Source/Text3D/Public/Extensions/Text3DDefaultLayoutExtension.h b/Engine/Plugins/VirtualProduction/Text3D/Source/Text3D/Public/Extensions/Text3DDefaultLayoutExtension.h
--- a/Engine/Plugins/VirtualProduction/Text3D/Source/Text3D/Public/Extensions/Text3DDefaultLayoutExtension.h
+++ b/Engine/Plugins/VirtualProduction/Text3D/Source/Text3D/Public/Extensions/Text3DDefaultLayoutExtension.h
@@ -242,4 +242,7 @@ private:
 
 	/** Determines how text is laid out, ie. parsing line breaks. */
 	TSharedPtr<FRichTextLayoutMarshaller> TextLayoutMarshaller;
+
+	/** Temporary storage for the glyph handles while text is built */
+	TArray<UE::Text3D::Geometry::FCachedFontFaceGlyphHandle> GlyphHandles;
 };
\ No newline at end of file
diff --git a/Engine/Plugins/VirtualProduction/Text3D/Source/Text3D/Public/IText3DBuildSystem.h b/Engine/Plugins/VirtualProduction/Text3D/Source/Text3D/Public/IText3DBuildSystem.h
--- a/Engine/Plugins/VirtualProduction/Text3D/Source/Text3D/Public/IText3DBuildSystem.h
+++ b/Engine/Plugins/VirtualProduction/Text3D/Source/Text3D/Public/IText3DBuildSystem.h
@@ -19,7 +19,7 @@ class IText3DBuildSystemInterface
 
 public:
 	/** Returns true if there are builders yet to complete */
-	virtual bool IsPendingBuildCompletion() const = 0;
+	virtual bool IsBuildInProgress() const = 0;
 
 	/** Adds the Text3D-managed primitives that are being built to the scene view's hidden primitive list */
 	virtual void HidePrimitivesBeingBuilt(FSceneView& InSceneView) const = 0;
