Motion Design: move text builder to happen during world ticking

#Synced-CL 51441789

diff --git a/Engine/Plugins/VirtualProduction/Text3D/Source/Text3D/Private/Subsystems/Text3DBuildSystem.cpp b/Engine/Plugins/VirtualProduction/Text3D/Source/Text3D/Private/Subsystems/Text3DBuildSystem.cpp
--- a/Engine/Plugins/VirtualProduction/Text3D/Source/Text3D/Private/Subsystems/Text3DBuildSystem.cpp
+++ b/Engine/Plugins/VirtualProduction/Text3D/Source/Text3D/Private/Subsystems/Text3DBuildSystem.cpp
@@ -26,7 +26,7 @@ TAutoConsoleVariable<int32> CVarFrameTimeBudgetMode(TEXT("Text3D.Build.FrameTime
 	);
 
 TAutoConsoleVariable<float> CVarFrameTimeBudget(TEXT("Text3D.Build.FrameTimeBudget")
-	, 0.014f
+	, 0.013f // default to 13 ms to account for the build time budget (default 2 ms) and the rest of the frame
 	, TEXT("The frame time cap to build Text3D components.\n")
 	  TEXT("Only relevant if Text3DBuild.FrameTimeBudgetMode is non-zero.\n")
 	  TEXT("If the current frame duration prior to the incremental build is past this limit, the build will be skipped (unless it has starved for more than the starve limit).\n")
@@ -185,6 +185,7 @@ void UText3DBuildSystem::RequestTextBuild(TNotNull<UText3DComponent*> InComponen
 	}
 
 	ConditionallyRegisterDelegates();
+	bShouldTick = true;
 
 	if (bInBlockingBuild)
 	{
@@ -250,8 +251,7 @@ void UText3DBuildSystem::IncrementalTextBuild()
 	if (!IsBuildingText())
 	{
 		UE_LOG(LogText3D, VeryVerbose, TEXT("%sIncremental Text3D build finished. No more build jobs left."), UE::Text3D::FScopedLog().GetLogPrefix());
-		FCoreDelegates::OnEndFrame.Remove(EndFrameHandle);
-		EndFrameHandle.Reset();
+		bShouldTick = false;
 		return; // nothing to update
 	}
 
@@ -348,6 +348,21 @@ bool UText3DBuildSystem::DoesSupportWorldType(const EWorldType::Type InWorldType
 	return InWorldType != EWorldType::None;
 }
 
+bool UText3DBuildSystem::IsTickable() const
+{
+	return bShouldTick;
+}
+
+void UText3DBuildSystem::Tick(float InDeltaTime)
+{
+	IncrementalTextBuild();
+}
+
+TStatId UText3DBuildSystem::GetStatId() const
+{
+	RETURN_QUICK_DECLARE_CYCLE_STAT(UText3DBuildSystem, STATGROUP_Tickables);
+}
+
 bool UText3DBuildSystem::IsBuildInProgress() const
 {
 	return IsBuildingText();
@@ -391,18 +406,10 @@ void UText3DBuildSystem::ConditionallyRegisterDelegates()
 				BeginFrameCycles = FPlatformTime::Cycles64();
 			});
 	}
-
-	if (!EndFrameHandle.IsValid())
-	{
-		EndFrameHandle = FCoreDelegates::OnEndFrame.AddUObject(this, &UText3DBuildSystem::IncrementalTextBuild);
-	}
 }
 
 void UText3DBuildSystem::UnregisterDelegates()
 {
 	FCoreDelegates::OnBeginFrame.Remove(BeginFrameHandle);
 	BeginFrameHandle.Reset();
-
-	FCoreDelegates::OnEndFrame.Remove(EndFrameHandle);
-	EndFrameHandle.Reset();
 }
diff --git a/Engine/Plugins/VirtualProduction/Text3D/Source/Text3D/Private/Subsystems/Text3DBuildSystem.h b/Engine/Plugins/VirtualProduction/Text3D/Source/Text3D/Private/Subsystems/Text3DBuildSystem.h
--- a/Engine/Plugins/VirtualProduction/Text3D/Source/Text3D/Private/Subsystems/Text3DBuildSystem.h
+++ b/Engine/Plugins/VirtualProduction/Text3D/Source/Text3D/Private/Subsystems/Text3DBuildSystem.h
@@ -21,7 +21,7 @@ enum class EText3DRendererFlags : uint8;
  * It's different from UText3DEngineSubsystem to keep these text builders referencing the text components scoped within the same world.
  */
 UCLASS(BlueprintType)
-class UText3DBuildSystem : public UWorldSubsystem, public IText3DBuildSystemInterface
+class UText3DBuildSystem : public UTickableWorldSubsystem, public IText3DBuildSystemInterface
 {
 	GENERATED_BODY()
 
@@ -56,6 +56,12 @@ protected:
 	virtual bool DoesSupportWorldType(const EWorldType::Type InWorldType) const override;
 	//~ End UWorldSubsystem
 
+	//~ Begin FTickableGameObject
+	virtual bool IsTickable() const override;
+	virtual void Tick(float InDeltaTime) override;
+	virtual TStatId GetStatId() const override;
+	//~ End FTickableGameObject
+
 	//~ Begin IText3DBuildSystemInterface
 	virtual bool IsBuildInProgress() const override;
 	virtual void HidePrimitivesBeingBuilt(FSceneView& InSceneView) const override;
@@ -86,6 +92,6 @@ private:
 	/** Handle to the begin frame call to keep track how long this frame's been going for */
 	FDelegateHandle BeginFrameHandle;
 
-	/** Handle to the end frame call for incremental builds */
-	FDelegateHandle EndFrameHandle;
+	/** Whether incremental build is required */
+	bool bShouldTick = false;
 };
