Text3D: changed verbosity from warning to verbose + added trace event scope

#Synced-CL 51441789

diff --git a/Engine/Plugins/VirtualProduction/Text3D/Source/Text3D/Private/Renderers/StaticMeshes/Text3DStaticMeshesRenderer.cpp b/Engine/Plugins/VirtualProduction/Text3D/Source/Text3D/Private/Renderers/StaticMeshes/Text3DStaticMeshesRenderer.cpp
--- a/Engine/Plugins/VirtualProduction/Text3D/Source/Text3D/Private/Renderers/StaticMeshes/Text3DStaticMeshesRenderer.cpp
+++ b/Engine/Plugins/VirtualProduction/Text3D/Source/Text3D/Private/Renderers/StaticMeshes/Text3DStaticMeshesRenderer.cpp
@@ -232,12 +232,12 @@ void UText3DStaticMeshesRenderer::OnUpdate(const UE::Text3D::Renderer::FUpdatePa
 				}
 				else
 				{
-					UE_LOG(LogText3D, Warning, TEXT("Invalid component object for character %i"), InIndex)
+					UE_LOG(LogText3D, Verbose, TEXT("Invalid component object for character %i"), InIndex)
 				}
 			}
 			else
 			{
-				UE_LOG(LogText3D, Warning, TEXT("Invalid component index for character %i"), InIndex)
+				UE_LOG(LogText3D, Verbose, TEXT("Invalid component index for character %i"), InIndex)
 			}
 		});
 	}
diff --git a/Engine/Plugins/VirtualProduction/Text3D/Source/Text3D/Private/Text3DBuilder.cpp b/Engine/Plugins/VirtualProduction/Text3D/Source/Text3D/Private/Text3DBuilder.cpp
--- a/Engine/Plugins/VirtualProduction/Text3D/Source/Text3D/Private/Text3DBuilder.cpp
+++ b/Engine/Plugins/VirtualProduction/Text3D/Source/Text3D/Private/Text3DBuilder.cpp
@@ -439,6 +439,7 @@ EBuildResult FTextBuilder::Next()
 
 void FTextBuilder::Finish()
 {
+	TRACE_CPUPROFILER_EVENT_SCOPE(Text3D::FinishBuild);
 	Component->FinishBuild(Flags);
 
 	Extensions.Reset();
