Text3D: Add blocking build option to text components + exposed this option to the build system and blueprints

#Synced-CL 51441789

diff --git a/Engine/Plugins/VirtualProduction/Text3D/Source/Text3D/Private/Extensions/Text3DDefaultStyleExtension.cpp b/Engine/Plugins/VirtualProduction/Text3D/Source/Text3D/Private/Extensions/Text3DDefaultStyleExtension.cpp
--- a/Engine/Plugins/VirtualProduction/Text3D/Source/Text3D/Private/Extensions/Text3DDefaultStyleExtension.cpp
+++ b/Engine/Plugins/VirtualProduction/Text3D/Source/Text3D/Private/Extensions/Text3DDefaultStyleExtension.cpp
@@ -36,7 +36,7 @@ void UText3DDefaultStyleExtension::SetStyleSet(UText3DStyleSet* InStyleSet)
 void UText3DDefaultStyleExtension::PostEditUndo()
 {
 	Super::PostEditUndo();
-	RequestUpdate(EText3DRendererFlags::All, /** Immediate */false);
+	RequestUpdate(EText3DRendererFlags::All);
 }
 
 void UText3DDefaultStyleExtension::PostEditChangeProperty(FPropertyChangedEvent& InPropertyChangedEvent)
@@ -50,7 +50,7 @@ void UText3DDefaultStyleExtension::PostEditChangeProperty(FPropertyChangedEvent&
 		if (InPropertyChangedEvent.ChangeType == EPropertyChangeType::ArrayClear
 			|| InPropertyChangedEvent.ChangeType == EPropertyChangeType::ArrayRemove)
 		{
-			RequestUpdate(EText3DRendererFlags::All, /** Immediate */false);
+			RequestUpdate(EText3DRendererFlags::All);
 		}
 	}
 	else if (MemberPropertyName == GET_MEMBER_NAME_CHECKED(UText3DDefaultStyleExtension, StyleSet))
@@ -191,5 +191,5 @@ void UText3DDefaultStyleExtension::OnStyleSetUpdated(UText3DStyleSet* InStyleSet
 
 void UText3DDefaultStyleExtension::OnStyleSetChanged()
 {
-	RequestUpdate(EText3DRendererFlags::All, /** Immediate */false);
+	RequestUpdate(EText3DRendererFlags::All);
 }
diff --git a/Engine/Plugins/VirtualProduction/Text3D/Source/Text3D/Private/Extensions/Text3DDefaultTokenExtension.cpp b/Engine/Plugins/VirtualProduction/Text3D/Source/Text3D/Private/Extensions/Text3DDefaultTokenExtension.cpp
--- a/Engine/Plugins/VirtualProduction/Text3D/Source/Text3D/Private/Extensions/Text3DDefaultTokenExtension.cpp
+++ b/Engine/Plugins/VirtualProduction/Text3D/Source/Text3D/Private/Extensions/Text3DDefaultTokenExtension.cpp
@@ -10,7 +10,7 @@
 void UText3DDefaultTokenExtension::PostEditUndo()
 {
 	Super::PostEditUndo();
-	RequestUpdate(EText3DRendererFlags::All, /** Immediate */false);
+	RequestUpdate(EText3DRendererFlags::All);
 }
 
 void UText3DDefaultTokenExtension::PostEditChangeProperty(FPropertyChangedEvent& InPropertyChangedEvent)
@@ -25,7 +25,7 @@ void UText3DDefaultTokenExtension::PostEditChangeProperty(FPropertyChangedEvent&
 			|| InPropertyChangedEvent.ChangeType == EPropertyChangeType::ArrayRemove
 			|| InPropertyChangedEvent.ChangeType == EPropertyChangeType::ArrayMove)
 		{
-			RequestUpdate(EText3DRendererFlags::All, /** Immediate */false);
+			RequestUpdate(EText3DRendererFlags::All);
 		}
 	}
 }
diff --git a/Engine/Plugins/VirtualProduction/Text3D/Source/Text3D/Private/Extensions/Text3DExtensionBase.cpp b/Engine/Plugins/VirtualProduction/Text3D/Source/Text3D/Private/Extensions/Text3DExtensionBase.cpp
--- a/Engine/Plugins/VirtualProduction/Text3D/Source/Text3D/Private/Extensions/Text3DExtensionBase.cpp
+++ b/Engine/Plugins/VirtualProduction/Text3D/Source/Text3D/Private/Extensions/Text3DExtensionBase.cpp
@@ -4,11 +4,11 @@
 
 #include "Text3DComponent.h"
 
-void UText3DExtensionBase::RequestUpdate(EText3DRendererFlags InFlags, bool bInImmediate) const
+void UText3DExtensionBase::RequestUpdate(EText3DRendererFlags InFlags, TOptional<bool> bInBlockingBuild) const
 {
 	if (UText3DComponent* Text3DComponent = GetText3DComponent())
 	{
-		Text3DComponent->RequestUpdate(InFlags, bInImmediate);
+		Text3DComponent->RequestUpdate(InFlags, bInBlockingBuild);
 	}
 }
 
diff --git a/Engine/Plugins/VirtualProduction/Text3D/Source/Text3D/Private/Styles/Text3DStyleBase.cpp b/Engine/Plugins/VirtualProduction/Text3D/Source/Text3D/Private/Styles/Text3DStyleBase.cpp
--- a/Engine/Plugins/VirtualProduction/Text3D/Source/Text3D/Private/Styles/Text3DStyleBase.cpp
+++ b/Engine/Plugins/VirtualProduction/Text3D/Source/Text3D/Private/Styles/Text3DStyleBase.cpp
@@ -163,7 +163,7 @@ void UText3DStyleBase::OnStylePropertiesChanged()
 	{
 		if (UText3DComponent* Text3DComponent = GetTypedOuter<UText3DComponent>())
 		{
-			Text3DComponent->RequestUpdate(EText3DRendererFlags::All, /** Immediate */false);
+			Text3DComponent->RequestUpdate(EText3DRendererFlags::All);
 		}
 		else if (UText3DStyleSet* StyleSet = GetTypedOuter<UText3DStyleSet>())
 		{
diff --git a/Engine/Plugins/VirtualProduction/Text3D/Source/Text3D/Private/Subsystems/Text3DBuildSystem.cpp b/Engine/Plugins/VirtualProduction/Text3D/Source/Text3D/Private/Subsystems/Text3DBuildSystem.cpp
--- a/Engine/Plugins/VirtualProduction/Text3D/Source/Text3D/Private/Subsystems/Text3DBuildSystem.cpp
+++ b/Engine/Plugins/VirtualProduction/Text3D/Source/Text3D/Private/Subsystems/Text3DBuildSystem.cpp
@@ -251,6 +251,8 @@ void UText3DBuildSystem::IncrementalTextBuild()
 	if (!IsBuildingText())
 	{
 		UE_LOG(LogText3D, VeryVerbose, TEXT("%sIncremental Text3D build finished. No more build jobs left."), UE::Text3D::FScopedLog().GetLogPrefix());
+		TextBuilders.Empty();
+		QueuedBuilders.Empty();
 		bShouldTick = false;
 		return; // nothing to update
 	}
@@ -293,6 +295,14 @@ void UText3DBuildSystem::IncrementalTextBuild()
 	}
 }
 
+void UText3DBuildSystem::BlockBuildTexts()
+{
+	Algo::ForEach(TextBuilders, &UE::Text3D::FTextBuilder::BlockingBuild);
+	Algo::ForEach(QueuedBuilders, &UE::Text3D::FTextBuilder::BlockingBuild);
+	TextBuilders.Empty();
+	QueuedBuilders.Empty();
+}
+
 bool UText3DBuildSystem::IsBuildingText() const
 {
 	if (!QueuedBuilders.IsEmpty() || !TextBuilders.IsEmpty())
@@ -353,6 +363,11 @@ bool UText3DBuildSystem::IsTickable() const
 	return bShouldTick;
 }
 
+bool UText3DBuildSystem::IsTickableInEditor() const
+{
+	return true;
+}
+
 void UText3DBuildSystem::Tick(float InDeltaTime)
 {
 	IncrementalTextBuild();
@@ -368,6 +383,11 @@ bool UText3DBuildSystem::IsBuildInProgress() const
 	return IsBuildingText();
 }
 
+void UText3DBuildSystem::FlushBuilds()
+{
+	BlockBuildTexts();
+}
+
 void UText3DBuildSystem::HidePrimitivesBeingBuilt(FSceneView& InSceneView) const
 {
 	auto AddPendingBuildItem = 
diff --git a/Engine/Plugins/VirtualProduction/Text3D/Source/Text3D/Private/Subsystems/Text3DBuildSystem.h b/Engine/Plugins/VirtualProduction/Text3D/Source/Text3D/Private/Subsystems/Text3DBuildSystem.h
--- a/Engine/Plugins/VirtualProduction/Text3D/Source/Text3D/Private/Subsystems/Text3DBuildSystem.h
+++ b/Engine/Plugins/VirtualProduction/Text3D/Source/Text3D/Private/Subsystems/Text3DBuildSystem.h
@@ -42,6 +42,13 @@ public:
 	/** Called on tick to process builders and call IncrementalBuild on them */
 	void IncrementalTextBuild();
 
+	/**
+	 * Executes any remaining text builds with no time budget. 
+	 * Warning: This can cause a hitch and should only be used when the result is needed immediately regardless of hitch.
+	 */
+	UFUNCTION(BlueprintCallable, Category="Text 3D")
+	void BlockBuildTexts();
+
 	/** Returns true if there are Text3D components being built */
 	UFUNCTION(BlueprintPure, Category="Text 3D")
 	bool IsBuildingText() const;
@@ -58,12 +65,14 @@ protected:
 
 	//~ Begin FTickableGameObject
 	virtual bool IsTickable() const override;
+	virtual bool IsTickableInEditor() const override;
 	virtual void Tick(float InDeltaTime) override;
 	virtual TStatId GetStatId() const override;
 	//~ End FTickableGameObject
 
 	//~ Begin IText3DBuildSystemInterface
 	virtual bool IsBuildInProgress() const override;
+	virtual void FlushBuilds() override;
 	virtual void HidePrimitivesBeingBuilt(FSceneView& InSceneView) const override;
 	//~ End IText3DBuildSystemInterface
 
diff --git a/Engine/Plugins/VirtualProduction/Text3D/Source/Text3D/Private/Text3DComponent.cpp b/Engine/Plugins/VirtualProduction/Text3D/Source/Text3D/Private/Text3DComponent.cpp
--- a/Engine/Plugins/VirtualProduction/Text3D/Source/Text3D/Private/Text3DComponent.cpp
+++ b/Engine/Plugins/VirtualProduction/Text3D/Source/Text3D/Private/Text3DComponent.cpp
@@ -85,10 +85,13 @@ void UText3DComponent::PostLoad()
 {
 	Super::PostLoad();
 	RefreshTypeface();
-	RequestUpdate(EText3DRendererFlags::All);
+
+	// Never perform immediate blocking builds in PostLoad
+	const bool bBlockingBuild = false;
+	RequestUpdate(EText3DRendererFlags::All, bBlockingBuild);
 }
 
-void UText3DComponent::RequestUpdate(EText3DRendererFlags InFlags, bool bInImmediate)
+void UText3DComponent::RequestUpdate(EText3DRendererFlags InFlags, TOptional<bool> bInBlockingBuild)
 {
 	if (InFlags != EText3DRendererFlags::None && !IsTemplate(RF_ClassDefaultObject))
 	{
@@ -96,7 +99,8 @@ void UText3DComponent::RequestUpdate(EText3DRendererFlags InFlags, bool bInImmed
 		{
 			if (UText3DBuildSystem* const BuildSystem = UText3DBuildSystem::Get(this))
 			{
-				BuildSystem->RequestTextBuild(this, InFlags, bInImmediate);
+				const bool bBlockingBuild = bInBlockingBuild.Get(bUseBlockingBuild);
+				BuildSystem->RequestTextBuild(this, InFlags, bBlockingBuild);
 			}
 		}
 		else
@@ -1229,6 +1233,11 @@ void UText3DComponent::CancelBuild()
 	bIsUpdatingText = false;
 }
 
+void UText3DComponent::SetUseBlockingBuild(bool bInUseBlockingBuild)
+{
+	bUseBlockingBuild = bInUseBlockingBuild;
+}
+
 void UText3DComponent::OnVisibilityChanged()
 {
 	Super::OnVisibilityChanged();
diff --git a/Engine/Plugins/VirtualProduction/Text3D/Source/Text3D/Private/Tokens/Text3DTokenBase.cpp b/Engine/Plugins/VirtualProduction/Text3D/Source/Text3D/Private/Tokens/Text3DTokenBase.cpp
--- a/Engine/Plugins/VirtualProduction/Text3D/Source/Text3D/Private/Tokens/Text3DTokenBase.cpp
+++ b/Engine/Plugins/VirtualProduction/Text3D/Source/Text3D/Private/Tokens/Text3DTokenBase.cpp
@@ -47,7 +47,7 @@ void UText3DTokenBase::OnTokenPropertiesChanged()
 	{
 		if (UText3DComponent* Text3DComponent = GetTypedOuter<UText3DComponent>())
 		{
-			Text3DComponent->RequestUpdate(EText3DRendererFlags::All, /** Immediate */false);
+			Text3DComponent->RequestUpdate(EText3DRendererFlags::All);
 		}
 	}
 }
diff --git a/Engine/Plugins/VirtualProduction/Text3D/Source/Text3D/Public/Extensions/Text3DExtensionBase.h b/Engine/Plugins/VirtualProduction/Text3D/Source/Text3D/Public/Extensions/Text3DExtensionBase.h
--- a/Engine/Plugins/VirtualProduction/Text3D/Source/Text3D/Public/Extensions/Text3DExtensionBase.h
+++ b/Engine/Plugins/VirtualProduction/Text3D/Source/Text3D/Public/Extensions/Text3DExtensionBase.h
@@ -62,7 +62,7 @@ public:
 	TEXT3D_API UText3DComponent* GetText3DComponent() const;
 
 protected:
-	void RequestUpdate(EText3DRendererFlags InFlags, bool bInImmediate = false) const;
+	void RequestUpdate(EText3DRendererFlags InFlags, TOptional<bool> bInBlockingBuild = {}) const;
 
 private:
 	uint16 UpdatePriority = 0;
diff --git a/Engine/Plugins/VirtualProduction/Text3D/Source/Text3D/Public/IText3DBuildSystem.h b/Engine/Plugins/VirtualProduction/Text3D/Source/Text3D/Public/IText3DBuildSystem.h
--- a/Engine/Plugins/VirtualProduction/Text3D/Source/Text3D/Public/IText3DBuildSystem.h
+++ b/Engine/Plugins/VirtualProduction/Text3D/Source/Text3D/Public/IText3DBuildSystem.h
@@ -21,6 +21,12 @@ public:
 	/** Returns true if there are builders yet to complete */
 	virtual bool IsBuildInProgress() const = 0;
 
+	/**
+	 * Executes any remaining text builds with no time budget. 
+	 * Warning: This can cause a hitch and should only be used when the result is needed immediately regardless of hitch. 
+	 */
+	virtual void FlushBuilds() = 0;
+
 	/** Adds the Text3D-managed primitives that are being built to the scene view's hidden primitive list */
 	virtual void HidePrimitivesBeingBuilt(FSceneView& InSceneView) const = 0;
 };
diff --git a/Engine/Plugins/VirtualProduction/Text3D/Source/Text3D/Public/Text3DComponent.h b/Engine/Plugins/VirtualProduction/Text3D/Source/Text3D/Public/Text3DComponent.h
--- a/Engine/Plugins/VirtualProduction/Text3D/Source/Text3D/Public/Text3DComponent.h
+++ b/Engine/Plugins/VirtualProduction/Text3D/Source/Text3D/Public/Text3DComponent.h
@@ -280,7 +280,7 @@ public:
 		return Statistics;
 	}
 
-	TEXT3D_API void RequestUpdate(EText3DRendererFlags InFlags, bool bInImmediate = false);
+	TEXT3D_API void RequestUpdate(EText3DRendererFlags InFlags, TOptional<bool> bInBlockingBuild = {});
 
 	int32 GetTypeFaceIndex() const;
 
@@ -427,6 +427,15 @@ public:
 	/** Called when the build is cancelled */
 	void CancelBuild();
 
+	UFUNCTION(BlueprintPure, BlueprintGetter, Category = "Text3D")
+	bool GetUseBlockingBuild() const
+	{
+		return bUseBlockingBuild;
+	}
+
+	UFUNCTION(BlueprintCallable, BlueprintSetter, Category = "Text3D")
+	TEXT3D_API void SetUseBlockingBuild(bool bInUseBlockingBuild);
+
 protected:
 	//~ Begin USceneComponent
 	/** Intercept and propagate a change on this component to all children. */
@@ -520,6 +529,10 @@ protected:
 	UPROPERTY(EditAnywhere, Getter="GetEnforceUpperCase", Setter="SetEnforceUpperCase", BlueprintSetter="SetEnforceUpperCase", BlueprintGetter="GetEnforceUpperCase", Category = "Text", meta = (AllowPrivateAccess="true"))
 	bool bEnforceUpperCase = false;
 
+	/** Whether to use blocking builds for text by default. Can be overriden by the RequestUpdate call */
+	UPROPERTY(EditAnywhere, Getter="GetUseBlockingBuild", BlueprintGetter="GetUseBlockingBuild", Setter="SetUseBlockingBuild", BlueprintSetter="SetUseBlockingBuild", Category="Text", meta=(AllowPrivateAccess="true"))
+	bool bUseBlockingBuild = false;
+
 	UPROPERTY(VisibleAnywhere, Instanced, Category="Token", meta=(AllowPrivateAccess="true"))
 	TObjectPtr<UText3DTokenExtensionBase> TokenExtension;
 
