Text3D: change tick timing by moving to a null object context for Game Tickable Object

#Synced-CL 51441789

diff --git a/Engine/Plugins/VirtualProduction/Text3D/Source/Text3D/Private/Subsystems/Text3DBuildSystem.h b/Engine/Plugins/VirtualProduction/Text3D/Source/Text3D/Private/Subsystems/Text3DBuildSystem.h
--- a/Engine/Plugins/VirtualProduction/Text3D/Source/Text3D/Private/Subsystems/Text3DBuildSystem.h
+++ b/Engine/Plugins/VirtualProduction/Text3D/Source/Text3D/Private/Subsystems/Text3DBuildSystem.h
@@ -21,7 +21,7 @@ enum class EText3DRendererFlags : uint8;
  * It's different from UText3DEngineSubsystem to keep these text builders referencing the text components scoped within the same world.
  */
 UCLASS(BlueprintType)
-class UText3DBuildSystem : public UTickableWorldSubsystem, public IText3DBuildSystemInterface
+class UText3DBuildSystem : public UWorldSubsystem, public FTickableGameObject, public IText3DBuildSystemInterface
 {
 	GENERATED_BODY()
 
