Text3D: use const ptr instead of ptr + exposed collision enabled for text to blueprints

#Synced-CL 51764840

diff --git a/Engine/Plugins/VirtualProduction/Text3D/Source/Text3D/Private/Utilities/Text3DUtilities.cpp b/Engine/Plugins/VirtualProduction/Text3D/Source/Text3D/Private/Utilities/Text3DUtilities.cpp
--- a/Engine/Plugins/VirtualProduction/Text3D/Source/Text3D/Private/Utilities/Text3DUtilities.cpp
+++ b/Engine/Plugins/VirtualProduction/Text3D/Source/Text3D/Private/Utilities/Text3DUtilities.cpp
@@ -35,7 +35,7 @@ THIRD_PARTY_INCLUDES_END
 
 #define LOCTEXT_NAMESPACE "Text3DUtilities"
 
-IText3DBuildSystemInterface* UE::Text3D::Utilities::FindBuildSystem(UObject* InWorldContextObject)
+IText3DBuildSystemInterface* UE::Text3D::Utilities::FindBuildSystem(const UObject* InWorldContextObject)
 {
 	if (!GEngine)
 	{
diff --git a/Engine/Plugins/VirtualProduction/Text3D/Source/Text3D/Public/Extensions/Text3DDefaultGeometryExtension.h b/Engine/Plugins/VirtualProduction/Text3D/Source/Text3D/Public/Extensions/Text3DDefaultGeometryExtension.h
--- a/Engine/Plugins/VirtualProduction/Text3D/Source/Text3D/Public/Extensions/Text3DDefaultGeometryExtension.h
+++ b/Engine/Plugins/VirtualProduction/Text3D/Source/Text3D/Public/Extensions/Text3DDefaultGeometryExtension.h
@@ -101,9 +101,11 @@ public:
 	{
 		return PivotVAlignment;
 	}
-	
-	UFUNCTION()
+
+	UFUNCTION(BlueprintCallable, Category = "Text3D|Geometry")
 	TEXT3D_API void SetCollisionEnabled(ECollisionEnabled::Type InCollisionEnabled);
+
+	UFUNCTION(BlueprintPure, Category = "Text3D|Geometry")
 	ECollisionEnabled::Type GetCollisionEnabled() const
 	{
 		return CollisionEnabled;
diff --git a/Engine/Plugins/VirtualProduction/Text3D/Source/Text3D/Public/Utilities/Text3DUtilities.h b/Engine/Plugins/VirtualProduction/Text3D/Source/Text3D/Public/Utilities/Text3DUtilities.h
--- a/Engine/Plugins/VirtualProduction/Text3D/Source/Text3D/Public/Utilities/Text3DUtilities.h
+++ b/Engine/Plugins/VirtualProduction/Text3D/Source/Text3D/Public/Utilities/Text3DUtilities.h
@@ -17,7 +17,7 @@ struct FTypefaceEntry;
 namespace UE::Text3D::Utilities
 {
 	/** Finds the text build system for the given world context object */
-	TEXT3D_API IText3DBuildSystemInterface* FindBuildSystem(UObject* InWorldContextObject);
+	TEXT3D_API IText3DBuildSystemInterface* FindBuildSystem(const UObject* InWorldContextObject);
 
 	namespace Font
 	{
