Text3D: changed way builder ticks by changing TickableType between Always and Never instead of IsTickable.

#Synced-CL 51764840

diff --git a/Engine/Plugins/VirtualProduction/Text3D/Source/Text3D/Private/Subsystems/Text3DBuildSystem.cpp b/Engine/Plugins/VirtualProduction/Text3D/Source/Text3D/Private/Subsystems/Text3DBuildSystem.cpp
--- a/Engine/Plugins/VirtualProduction/Text3D/Source/Text3D/Private/Subsystems/Text3DBuildSystem.cpp
+++ b/Engine/Plugins/VirtualProduction/Text3D/Source/Text3D/Private/Subsystems/Text3DBuildSystem.cpp
@@ -87,6 +87,11 @@ TSharedRef<FText3DSceneViewExtension> GetSceneViewExtension()
 
 } // UE::Text3D::Private
 
+UText3DBuildSystem::UText3DBuildSystem()
+	: FTickableGameObject(ETickableTickType::Never)
+{
+}
+
 UText3DBuildSystem* UText3DBuildSystem::Get(TNotNull<const UText3DComponent*> InComponent)
 {
 	if (const UWorld* const World = InComponent->GetWorld())
@@ -185,7 +190,7 @@ void UText3DBuildSystem::RequestTextBuild(TNotNull<UText3DComponent*> InComponen
 	}
 
 	ConditionallyRegisterDelegates();
-	bShouldTick = true;
+	SetTickableTickType(ETickableTickType::Always);
 
 	if (bInBlockingBuild)
 	{
@@ -253,7 +258,7 @@ void UText3DBuildSystem::IncrementalTextBuild()
 		UE_LOG(LogText3D, VeryVerbose, TEXT("%sIncremental Text3D build finished. No more build jobs left."), UE::Text3D::FScopedLog().GetLogPrefix());
 		TextBuilders.Empty();
 		QueuedBuilders.Empty();
-		bShouldTick = false;
+		SetTickableTickType(ETickableTickType::Never);
 		return; // nothing to update
 	}
 
@@ -358,11 +363,6 @@ bool UText3DBuildSystem::DoesSupportWorldType(const EWorldType::Type InWorldType
 	return InWorldType != EWorldType::None;
 }
 
-bool UText3DBuildSystem::IsTickable() const
-{
-	return bShouldTick;
-}
-
 bool UText3DBuildSystem::IsTickableInEditor() const
 {
 	return true;
diff --git a/Engine/Plugins/VirtualProduction/Text3D/Source/Text3D/Private/Subsystems/Text3DBuildSystem.h b/Engine/Plugins/VirtualProduction/Text3D/Source/Text3D/Private/Subsystems/Text3DBuildSystem.h
--- a/Engine/Plugins/VirtualProduction/Text3D/Source/Text3D/Private/Subsystems/Text3DBuildSystem.h
+++ b/Engine/Plugins/VirtualProduction/Text3D/Source/Text3D/Private/Subsystems/Text3DBuildSystem.h
@@ -26,6 +26,8 @@ class UText3DBuildSystem : public UWorldSubsystem, public FTickableGameObject, p
 	GENERATED_BODY()
 
 public:
+	UText3DBuildSystem();
+
 	static UText3DBuildSystem* Get(TNotNull<const UText3DComponent*> InComponent);
 
 	/** 
@@ -64,7 +66,6 @@ protected:
 	//~ End UWorldSubsystem
 
 	//~ Begin FTickableGameObject
-	virtual bool IsTickable() const override;
 	virtual bool IsTickableInEditor() const override;
 	virtual void Tick(float InDeltaTime) override;
 	virtual TStatId GetStatId() const override;
@@ -100,7 +101,4 @@ private:
 
 	/** Handle to the begin frame call to keep track how long this frame's been going for */
 	FDelegateHandle BeginFrameHandle;
-
-	/** Whether incremental build is required */
-	bool bShouldTick = false;
 };
