Text3D: fix issues where components outered to the meshes renderer were not being renamed and effectively not updating their cached owner.

#Synced-CL 52400047

diff --git a/Engine/Plugins/VirtualProduction/Text3D/Source/Text3D/Private/Renderers/DynamicMesh/Text3DDynamicMeshRenderer.cpp b/Engine/Plugins/VirtualProduction/Text3D/Source/Text3D/Private/Renderers/DynamicMesh/Text3DDynamicMeshRenderer.cpp
--- a/Engine/Plugins/VirtualProduction/Text3D/Source/Text3D/Private/Renderers/DynamicMesh/Text3DDynamicMeshRenderer.cpp
+++ b/Engine/Plugins/VirtualProduction/Text3D/Source/Text3D/Private/Renderers/DynamicMesh/Text3DDynamicMeshRenderer.cpp
@@ -14,6 +14,7 @@
 #include "Subsystems/Text3DEngineSubsystem.h"
 #include "Text3DComponent.h"
 #include "Text3DInternalTypes.h"
+#include "UObject/Package.h"
 
 UText3DDynamicMeshRenderer::FScopedMeshEditor::~FScopedMeshEditor()
 {
@@ -59,6 +60,7 @@ void UText3DDynamicMeshRenderer::OnCreate()
 	if (!DynamicMeshComponent)
 	{
 		DynamicMeshComponent = NewObject<UDynamicMeshComponent>(this, TEXT("Text3DDynamicMesh"));
+		DynamicMeshComponent->CreationMethod = EComponentCreationMethod::Instance;
 		DynamicMeshComponent->RegisterComponent();
 		Owner->AddOwnedComponent(DynamicMeshComponent);
 		DynamicMeshComponent->AttachToComponent(TextComponent, AttachRule);
@@ -249,6 +251,27 @@ void UText3DDynamicMeshRenderer::OnDebugModeDisabled()
 	SetDebugMode(false);
 }
 
+void UText3DDynamicMeshRenderer::PostRename(UObject* InOldOuter, const FName InOldName)
+{
+	Super::PostRename(InOldOuter, InOldName);
+
+	// Renderer's managed components are not outered to the Text3D Component directly,
+	// Use a temp outer to trigger a refresh on the components cached owner
+	UObject* const TempOuter = GetTransientPackage();
+
+	auto RenameComponent = [TempOuter, this](UActorComponent* InComponent)
+		{
+			if (InComponent)
+			{
+				// Force trigger PostRename on the components as its one of the few places where the cached owner is updated
+				InComponent->PostRename(TempOuter, InComponent->GetFName());
+				InComponent->PostRename(this, InComponent->GetFName());
+			}
+		};
+
+	RenameComponent(DynamicMeshComponent);
+}
+
 void UText3DDynamicMeshRenderer::SetDebugMode(bool bInEnabled)
 {
 	// Since we are dealing with class FProperty, no need to run this for each instance, do it once
diff --git a/Engine/Plugins/VirtualProduction/Text3D/Source/Text3D/Private/Renderers/DynamicMesh/Text3DDynamicMeshRenderer.h b/Engine/Plugins/VirtualProduction/Text3D/Source/Text3D/Private/Renderers/DynamicMesh/Text3DDynamicMeshRenderer.h
--- a/Engine/Plugins/VirtualProduction/Text3D/Source/Text3D/Private/Renderers/DynamicMesh/Text3DDynamicMeshRenderer.h
+++ b/Engine/Plugins/VirtualProduction/Text3D/Source/Text3D/Private/Renderers/DynamicMesh/Text3DDynamicMeshRenderer.h
@@ -70,6 +70,10 @@ protected:
 #endif
 	//~ End UText3DRendererBase
 
+	//~ Begin UObject
+	virtual void PostRename(UObject* InOldOuter, const FName InOldName) override;
+	//~ End UObject
+
 #if WITH_EDITOR
 	void SetDebugMode(bool bInEnabled);
 #endif
diff --git a/Engine/Plugins/VirtualProduction/Text3D/Source/Text3D/Private/Renderers/StaticMeshes/Text3DStaticMeshesRenderer.cpp b/Engine/Plugins/VirtualProduction/Text3D/Source/Text3D/Private/Renderers/StaticMeshes/Text3DStaticMeshesRenderer.cpp
--- a/Engine/Plugins/VirtualProduction/Text3D/Source/Text3D/Private/Renderers/StaticMeshes/Text3DStaticMeshesRenderer.cpp
+++ b/Engine/Plugins/VirtualProduction/Text3D/Source/Text3D/Private/Renderers/StaticMeshes/Text3DStaticMeshesRenderer.cpp
@@ -1,6 +1,8 @@
 // Copyright Epic Games, Inc. All Rights Reserved.
 
 #include "Renderers/StaticMeshes/Text3DStaticMeshesRenderer.h"
+
+#include "Algo/ForEach.h"
 #include "Characters/Text3DCharacterBase.h"
 #include "Components/SceneComponent.h"
 #include "Components/StaticMeshComponent.h"
@@ -370,7 +372,6 @@ void UText3DStaticMeshesRenderer::AllocateComponents(int32 InCount)
 				{
 					MeshComponent->UnregisterComponent();
 				}
-				Actor->RemoveInstanceComponent(MeshComponent);
 				Actor->RemoveOwnedComponent(MeshComponent);
 
 				if (MeshComponent->GetOwner() == Actor)
@@ -394,6 +395,7 @@ void UText3DStaticMeshesRenderer::AllocateComponents(int32 InCount)
 			if (!MeshComponent)
 			{
 				MeshComponent = NewObject<UStaticMeshComponent>(this);
+				MeshComponent->CreationMethod = EComponentCreationMethod::Instance;
 			}
 
 			if (MeshComponent)
@@ -469,6 +471,29 @@ void UText3DStaticMeshesRenderer::OnDebugModeDisabled()
 	SetDebugMode(false);
 }
 
+void UText3DStaticMeshesRenderer::PostRename(UObject* InOldOuter, const FName InOldName)
+{
+	Super::PostRename(InOldOuter, InOldName);
+
+	// Renderer's managed components are not outered to the Text3D Component directly,
+	// Use a temp outer to trigger a refresh on the components cached owner
+	UObject* const TempOuter = GetTransientPackage();
+
+	auto RenameComponent = [TempOuter, this](UActorComponent* InComponent)
+		{
+			if (InComponent)
+			{
+				// Force trigger PostRename on the components as its one of the few places where the cached owner is updated
+				InComponent->PostRename(TempOuter, InComponent->GetFName());
+				InComponent->PostRename(this, InComponent->GetFName());
+			}
+		};
+
+	RenameComponent(TextRoot);
+	Algo::ForEach(CharacterMeshes, RenameComponent);
+	Algo::ForEach(CharacterMeshesPool, RenameComponent);
+}
+
 void UText3DStaticMeshesRenderer::SetDebugMode(bool bEnabled)
 {
 	// Since we are dealing with class FProperty, no need to run this for each instance, do it once
diff --git a/Engine/Plugins/VirtualProduction/Text3D/Source/Text3D/Private/Renderers/StaticMeshes/Text3DStaticMeshesRenderer.h b/Engine/Plugins/VirtualProduction/Text3D/Source/Text3D/Private/Renderers/StaticMeshes/Text3DStaticMeshesRenderer.h
--- a/Engine/Plugins/VirtualProduction/Text3D/Source/Text3D/Private/Renderers/StaticMeshes/Text3DStaticMeshesRenderer.h
+++ b/Engine/Plugins/VirtualProduction/Text3D/Source/Text3D/Private/Renderers/StaticMeshes/Text3DStaticMeshesRenderer.h
@@ -60,6 +60,10 @@ protected:
 #endif
 	//~ End UText3DRendererBase
 
+	//~ Begin UObject
+	virtual void PostRename(UObject* InOldOuter, const FName InOldName) override;
+	//~ End UObject
+
 	void AllocateComponents(int32 InNewCharacterCount);
 
 #if WITH_EDITOR
