Text3D: Fix Inaccuracies from character pivot alignment

#Synced-CL 52400047

diff --git a/Engine/Plugins/VirtualProduction/Text3D/Source/Text3D/Private/Extensions/Text3DDefaultGeometryExtension.cpp b/Engine/Plugins/VirtualProduction/Text3D/Source/Text3D/Private/Extensions/Text3DDefaultGeometryExtension.cpp
--- a/Engine/Plugins/VirtualProduction/Text3D/Source/Text3D/Private/Extensions/Text3DDefaultGeometryExtension.cpp
+++ b/Engine/Plugins/VirtualProduction/Text3D/Source/Text3D/Private/Extensions/Text3DDefaultGeometryExtension.cpp
@@ -295,9 +295,9 @@ FVector UText3DDefaultGeometryExtension::GetPivotOffset() const
 	case EText3DHorizontalTextAlignment::Left:
 		return FVector::ZeroVector;
 	case EText3DHorizontalTextAlignment::Center:
-		return FVector(0.0f, -0.5f, 0.0f);
+		return FVector(0.0f, 0.5f, 0.0f);
 	case EText3DHorizontalTextAlignment::Right:
-		return FVector(0.0f, -1.f, 0.0f);
+		return FVector(0.0f, 1.f, 0.0f);
 	}
 
 	return FVector::ZeroVector;
diff --git a/Engine/Plugins/VirtualProduction/Text3D/Source/Text3D/Private/Extensions/Text3DDefaultLayoutExtension.cpp b/Engine/Plugins/VirtualProduction/Text3D/Source/Text3D/Private/Extensions/Text3DDefaultLayoutExtension.cpp
--- a/Engine/Plugins/VirtualProduction/Text3D/Source/Text3D/Private/Extensions/Text3DDefaultLayoutExtension.cpp
+++ b/Engine/Plugins/VirtualProduction/Text3D/Source/Text3D/Private/Extensions/Text3DDefaultLayoutExtension.cpp
@@ -427,20 +427,19 @@ EText3DExtensionResult UText3DDefaultLayoutExtension::PreRendererUpdate(const UE
 					return EText3DExtensionResult::Failed;
 				}
 
-				const FBox GlyphBounds = CachedMesh->MeshBounds;
-				const FVector GlyphSize = GlyphBounds.GetSize();
+				const FVector GlyphSize = CachedMesh->MeshBounds.GetSize();
+
+				const FVector2D FaceSize = FVector2D::Max(CachedMesh->FontFaceGlyphSize, FVector2D(UE_SMALL_NUMBER));
+				const FVector Ratio(1.f, GlyphEntry->Size / FaceSize.X, GlyphEntry->Size / FaceSize.Y);
 
 				FVector Shift = FVector::ZeroVector;
 				switch (GeometryExtension->GetGlyphHAlignment())
 				{
-				case EText3DHorizontalTextAlignment::Left:
-					Shift = FVector::ZeroVector;
-					break;
 				case EText3DHorizontalTextAlignment::Center:
-					Shift.Y = (GlyphSize.Y / 2);
+					Shift.Y += (0.5 * GlyphSize.Y + CachedMesh->MeshOffset.Y) * Ratio.Y;
 					break;
 				case EText3DHorizontalTextAlignment::Right:
-					Shift.Y = GlyphSize.Y;
+					Shift.Y += (GlyphSize.Y + CachedMesh->MeshOffset.Y) * Ratio.Y;
 					break;
 				}
 
@@ -457,17 +456,7 @@ EText3DExtensionResult UText3DDefaultLayoutExtension::PreRendererUpdate(const UE
 				}
 
 				FTransform Transform(Location + Shift);
-				const FVector2D FaceSize = CachedMesh->FontFaceGlyphSize;
-
-				if (FaceSize.IsNearlyZero())
-				{
-					UE_LOG(LogText3D, Warning, TEXT("Font face glyph size is nearly zero, cannot scale mesh properly to match font size"))
-				}
-				else
-				{
-					const FVector Ratio(1.f, GlyphEntry->Size / FaceSize.X, GlyphEntry->Size / FaceSize.Y);
-					Transform.SetScale3D(Ratio);
-				}
+				Transform.SetScale3D(Ratio);
 
 				constexpr bool bReset = true;
 				FTransform& CharacterTransform = Character->GetTransform(bReset);
diff --git a/Engine/Plugins/VirtualProduction/Text3D/Source/Text3D/Private/GeometryBuilders/Text3DGlyphMeshBuilder.cpp b/Engine/Plugins/VirtualProduction/Text3D/Source/Text3D/Private/GeometryBuilders/Text3DGlyphMeshBuilder.cpp
--- a/Engine/Plugins/VirtualProduction/Text3D/Source/Text3D/Private/GeometryBuilders/Text3DGlyphMeshBuilder.cpp
+++ b/Engine/Plugins/VirtualProduction/Text3D/Source/Text3D/Private/GeometryBuilders/Text3DGlyphMeshBuilder.cpp
@@ -114,18 +114,17 @@ void FText3DGlyphMeshBuilder::MovePivot(const FVector& InNewPivot)
 		MaxY = FMath::Max(MaxY, Position.Y);
 	}
 
-	// Adjust mesh to remove spaces around
-	const double CenterX = (MinX + MaxX) * 0.5f;
-	MeshOffset = FVector(-CenterX, -MinY, 0.0f);
+	const FVector MeshCenter = FVector(MinX+MaxX, MinY+MaxY, 0) * 0.5;
 
-	// Apply pivot offset with adjusted values
-	const FVector3f PivotOffset = FVector3f(InNewPivot) * FVector3f(MaxX - MinX, MaxY - MinY, 0);
+	// Keep track of the mesh offset that we need to add back.
+	MeshOffset  = +2.0 * InNewPivot * FVector(MinX, MinY, 0);
+	PivotOffset = -2.0 * InNewPivot * MeshCenter;
 
+	// Apply pivot offset with adjusted values
+	const FVector3f PivotOffsetF = FVector3f(PivotOffset);
 	for (const FVertexID VertexID : Data.GetGlyph().GetMeshDescription().Vertices().GetElementIDs())
 	{
-		FVector3f CurrentPosition = VertexPositions[VertexID];
-		CurrentPosition += PivotOffset;
-		VertexPositions[VertexID] = CurrentPosition;
+		VertexPositions[VertexID] += PivotOffsetF;
 	}
 }
 
diff --git a/Engine/Plugins/VirtualProduction/Text3D/Source/Text3D/Private/GeometryBuilders/Text3DGlyphMeshBuilder.h b/Engine/Plugins/VirtualProduction/Text3D/Source/Text3D/Private/GeometryBuilders/Text3DGlyphMeshBuilder.h
--- a/Engine/Plugins/VirtualProduction/Text3D/Source/Text3D/Private/GeometryBuilders/Text3DGlyphMeshBuilder.h
+++ b/Engine/Plugins/VirtualProduction/Text3D/Source/Text3D/Private/GeometryBuilders/Text3DGlyphMeshBuilder.h
@@ -47,6 +47,7 @@ public:
 
 private:
 	FVector MeshOffset = FVector::ZeroVector;
+	FVector PivotOffset = FVector::ZeroVector;
 	FText3DGlyphData Data;
 	FText3DGlyphContourList Contours;
 
