Text3D: fixed post rename not being defined outside editor builds

#Synced-CL 52400047

diff --git a/Engine/Plugins/VirtualProduction/Text3D/Source/Text3D/Private/Renderers/DynamicMesh/Text3DDynamicMeshRenderer.cpp b/Engine/Plugins/VirtualProduction/Text3D/Source/Text3D/Private/Renderers/DynamicMesh/Text3DDynamicMeshRenderer.cpp
--- a/Engine/Plugins/VirtualProduction/Text3D/Source/Text3D/Private/Renderers/DynamicMesh/Text3DDynamicMeshRenderer.cpp
+++ b/Engine/Plugins/VirtualProduction/Text3D/Source/Text3D/Private/Renderers/DynamicMesh/Text3DDynamicMeshRenderer.cpp
@@ -250,6 +250,7 @@ void UText3DDynamicMeshRenderer::OnDebugModeDisabled()
 	Super::OnDebugModeDisabled();
 	SetDebugMode(false);
 }
+#endif
 
 void UText3DDynamicMeshRenderer::PostRename(UObject* InOldOuter, const FName InOldName)
 {
@@ -272,6 +273,7 @@ void UText3DDynamicMeshRenderer::PostRename(UObject* InOldOuter, const FName InO
 	RenameComponent(DynamicMeshComponent);
 }
 
+#if WITH_EDITOR
 void UText3DDynamicMeshRenderer::SetDebugMode(bool bInEnabled)
 {
 	// Since we are dealing with class FProperty, no need to run this for each instance, do it once
diff --git a/Engine/Plugins/VirtualProduction/Text3D/Source/Text3D/Private/Renderers/StaticMeshes/Text3DStaticMeshesRenderer.cpp b/Engine/Plugins/VirtualProduction/Text3D/Source/Text3D/Private/Renderers/StaticMeshes/Text3DStaticMeshesRenderer.cpp
--- a/Engine/Plugins/VirtualProduction/Text3D/Source/Text3D/Private/Renderers/StaticMeshes/Text3DStaticMeshesRenderer.cpp
+++ b/Engine/Plugins/VirtualProduction/Text3D/Source/Text3D/Private/Renderers/StaticMeshes/Text3DStaticMeshesRenderer.cpp
@@ -470,6 +470,7 @@ void UText3DStaticMeshesRenderer::OnDebugModeDisabled()
 	Super::OnDebugModeDisabled();
 	SetDebugMode(false);
 }
+#endif
 
 void UText3DStaticMeshesRenderer::PostRename(UObject* InOldOuter, const FName InOldName)
 {
@@ -494,6 +495,7 @@ void UText3DStaticMeshesRenderer::PostRename(UObject* InOldOuter, const FName In
 	Algo::ForEach(CharacterMeshesPool, RenameComponent);
 }
 
+#if WITH_EDITOR
 void UText3DStaticMeshesRenderer::SetDebugMode(bool bEnabled)
 {
 	// Since we are dealing with class FProperty, no need to run this for each instance, do it once
