Text3D: Removed frame build budget, as it's likely to cause starved builders to execute all at once + added build budget percent for building glyphs.

#Synced-CL 52400047

diff --git a/Engine/Plugins/VirtualProduction/Text3D/Source/Text3D/Private/Extensions/Text3DDefaultLayoutExtension.cpp b/Engine/Plugins/VirtualProduction/Text3D/Source/Text3D/Private/Extensions/Text3DDefaultLayoutExtension.cpp
--- a/Engine/Plugins/VirtualProduction/Text3D/Source/Text3D/Private/Extensions/Text3DDefaultLayoutExtension.cpp
+++ b/Engine/Plugins/VirtualProduction/Text3D/Source/Text3D/Private/Extensions/Text3DDefaultLayoutExtension.cpp
@@ -376,11 +376,6 @@ EText3DExtensionResult UText3DDefaultLayoutExtension::PreRendererUpdate(const UE
 		return EText3DExtensionResult::Failed;
 	}
 
-	if (EnumHasAnyFlags(InParameters.UpdateFlags, EText3DRendererFlags::Geometry))
-	{
-		TextSubsystem->ProcessBuildGlyphMeshes();	
-	}
-
 	const FGlyphMeshParameters& GlyphMeshParameters = *GeometryExtension->GetGlyphMeshParameters();
 
 	CharacterIndex = 0;
diff --git a/Engine/Plugins/VirtualProduction/Text3D/Source/Text3D/Private/Subsystems/Text3DBuildSystem.cpp b/Engine/Plugins/VirtualProduction/Text3D/Source/Text3D/Private/Subsystems/Text3DBuildSystem.cpp
--- a/Engine/Plugins/VirtualProduction/Text3D/Source/Text3D/Private/Subsystems/Text3DBuildSystem.cpp
+++ b/Engine/Plugins/VirtualProduction/Text3D/Source/Text3D/Private/Subsystems/Text3DBuildSystem.cpp
@@ -12,27 +12,12 @@
 #include "SceneView.h"
 #include "SceneViewExtension.h"
 #include "Text3DComponent.h"
+#include "Text3DEngineSubsystem.h"
 #include "Text3DSceneViewExtension.h"
 #include "Utilities/Text3DScopedLog.h"
 
 namespace UE::Text3D::Private
 {
-TAutoConsoleVariable<int32> CVarFrameTimeBudgetMode(TEXT("Text3D.Build.FrameTimeBudgetMode")
-	, 1
-	, TEXT("The frame time budget mode to build Text3D components.\n")
-	  TEXT("0: Frame time budget disabled\n")
-	  TEXT("1: Frame time budget enabled only in playing worlds\n")
-	  TEXT("2: Frame time budget enabled in both playing worlds and editor worlds.\n")
-	);
-
-TAutoConsoleVariable<float> CVarFrameTimeBudget(TEXT("Text3D.Build.FrameTimeBudget")
-	, 0.013f // default to 13 ms to account for the build time budget (default 2 ms) and the rest of the frame
-	, TEXT("The frame time cap to build Text3D components.\n")
-	  TEXT("Only relevant if Text3DBuild.FrameTimeBudgetMode is non-zero.\n")
-	  TEXT("If the current frame duration prior to the incremental build is past this limit, the build will be skipped (unless it has starved for more than the starve limit).\n")
-	  TEXT("See Text3D.Build.StarveFrameLimit")
-	);
-
 TAutoConsoleVariable<int32> CVarStarveFrameLimit(TEXT("Text3D.Build.StarveFrameLimit")
 	, 10
 	, TEXT("The maximum number of frames that the build can starve before forcing an incremental build"));
@@ -46,31 +31,10 @@ TAutoConsoleVariable<float> CVarBuildTimeBudget(TEXT("Text3D.Build.BuildTimeBudg
 	, TEXT("The time budget allocated for building Text3D components.\n")
 	  TEXT("Requires Text3D.Build.EnableBuildTimeBudget to be true."));
 
-/** Returns true if time budget is enabled for the given world */
-bool IsFrameTimeBudgetEnabled(const UWorld* InWorld)
-{
-	if (!InWorld)
-	{
-		return false;
-	}
-
-	// Mode 0: Disabled
-	// Mode 1: Only playing worlds
-	// Mode 2: Any world
-	const int32 Mode = CVarFrameTimeBudgetMode.GetValueOnGameThread();
-	if (Mode <= 0 || Mode > 2)
-	{
-		return false;
-	}
-
-#if WITH_EDITOR
-	constexpr EWorldType::Type PlayingWorldTypes[] = { EWorldType::Game, EWorldType::GamePreview, EWorldType::PIE };
-	const bool bOnlyPlayingWorlds = Mode == 1;
-	return !bOnlyPlayingWorlds || Algo::Contains(PlayingWorldTypes, InWorld->WorldType);
-#else
-	return true;
-#endif
-}
+TAutoConsoleVariable<float> CVarBuildGlyphBudgetThreshold(TEXT("Text3D.Build.GlyphBuildBudgetThreshold")
+	, 0.25f
+	, TEXT("If the build budget used so far (as a %, from 0-1) exceeds this, the glyph building will take place in the next incremental build tick.\n")
+	  TEXT("Requires Text3D.Build.EnableBuildTimeBudget to be true."));
 
 /** Returns the scene view extension to use for this build system */
 TSharedRef<FText3DSceneViewExtension> GetSceneViewExtension()
@@ -262,6 +226,12 @@ void UText3DBuildSystem::IncrementalTextBuild()
 		return; // nothing to update
 	}
 
+	UText3DEngineSubsystem* const TextSubsystem = UText3DEngineSubsystem::Get();
+	if (!ensure(TextSubsystem))
+	{
+		return;
+	}
+
 	const uint64 CurrentCycles = FPlatformTime::Cycles64();
 
 	UE::Text3D::FTextBuilder::FIncrementalBuildParams Params;
@@ -270,7 +240,7 @@ void UText3DBuildSystem::IncrementalTextBuild()
 		Params.BudgetEnd = CurrentCycles + FPlatformTime::SecondsToCycles64(UE::Text3D::Private::CVarBuildTimeBudget.GetValueOnGameThread());
 		if (*Params.BudgetEnd <= BudgetDeficit)
 		{
-			*Params.BudgetEnd = 0;
+			*Params.BudgetEnd = CurrentCycles;
 		}
 		else
 		{
@@ -280,11 +250,10 @@ void UText3DBuildSystem::IncrementalTextBuild()
 	BudgetDeficit = 0;
 
 	Params.MaxStarvedFrames = Private::CVarStarveFrameLimit.GetValueOnGameThread();
-	Params.bFrameTimeExceeded = Private::IsFrameTimeBudgetEnabled(GetWorld()) 
-		&& (CurrentCycles - *BeginFrameCycles) >= FPlatformTime::SecondsToCycles64(Private::CVarFrameTimeBudget.GetValueOnGameThread());
-
 	Params.TempMaxUpdatePhase = EBuildPhase::PrepareExtensions;
 
+	const int32 GlyphMeshBuildCount = TextSubsystem->GetQueuedGlyphMeshBuildCount();
+
 	// Update all the builders that are preparing
 	for (UE::Text3D::FTextBuilder& Builder : TextBuilders)
 	{
@@ -293,6 +262,19 @@ void UText3DBuildSystem::IncrementalTextBuild()
 
 	Params.TempMaxUpdatePhase.Reset();
 
+	// If there are any new glyph mesh builds queued, make sure there's enough budget time to build these
+	if (GlyphMeshBuildCount != TextSubsystem->GetQueuedGlyphMeshBuildCount() && Params.BudgetEnd.IsSet())
+	{
+		// Calculate the amount of budget we've used so far (%)
+		const double BuildBudgetPct = FMath::SafeDivide(static_cast<double>(FPlatformTime::Cycles64() - CurrentCycles), static_cast<double>(*Params.BudgetEnd - CurrentCycles));
+		if (BuildBudgetPct >= Private::CVarBuildGlyphBudgetThreshold.GetValueOnGameThread())
+		{
+			return;
+		}
+	}
+
+	TextSubsystem->ProcessBuildGlyphMeshes();
+
 	// Update all the builders with no phase limit
 	for (UE::Text3D::FTextBuilder& Builder : TextBuilders)
 	{
diff --git a/Engine/Plugins/VirtualProduction/Text3D/Source/Text3D/Private/Subsystems/Text3DEngineSubsystem.cpp b/Engine/Plugins/VirtualProduction/Text3D/Source/Text3D/Private/Subsystems/Text3DEngineSubsystem.cpp
--- a/Engine/Plugins/VirtualProduction/Text3D/Source/Text3D/Private/Subsystems/Text3DEngineSubsystem.cpp
+++ b/Engine/Plugins/VirtualProduction/Text3D/Source/Text3D/Private/Subsystems/Text3DEngineSubsystem.cpp
@@ -248,6 +248,11 @@ void UText3DEngineSubsystem::QueueBuildGlyphMeshes(TArray<FText3DBuildGlyphMeshD
 	QueuedMeshDescs.Append(MoveTemp(InMeshDescs));
 }
 
+int32 UText3DEngineSubsystem::GetQueuedGlyphMeshBuildCount() const
+{
+	return QueuedMeshDescs.Num();
+}
+
 void UText3DEngineSubsystem::ProcessBuildGlyphMeshes()
 {
 	// If queued mesh desc is empty, the mesh parameters should also be.
diff --git a/Engine/Plugins/VirtualProduction/Text3D/Source/Text3D/Private/Subsystems/Text3DEngineSubsystem.h b/Engine/Plugins/VirtualProduction/Text3D/Source/Text3D/Private/Subsystems/Text3DEngineSubsystem.h
--- a/Engine/Plugins/VirtualProduction/Text3D/Source/Text3D/Private/Subsystems/Text3DEngineSubsystem.h
+++ b/Engine/Plugins/VirtualProduction/Text3D/Source/Text3D/Private/Subsystems/Text3DEngineSubsystem.h
@@ -132,6 +132,9 @@ public:
 	/** Adds the given mesh desc and glyph parameters to the queue */
 	void QueueBuildGlyphMeshes(TArray<FText3DBuildGlyphMeshDesc>&& InMeshDescs, const FGlyphMeshParameters& InGlyphMeshParameters);
 
+	/** Returns the number of glyph meshes pending build */
+	int32 GetQueuedGlyphMeshBuildCount() const;
+
 	/** Processes the queued mesh descs and glyph parameters across multiple threads */
 	void ProcessBuildGlyphMeshes();
 
diff --git a/Engine/Plugins/VirtualProduction/Text3D/Source/Text3D/Private/Text3DBuilder.cpp b/Engine/Plugins/VirtualProduction/Text3D/Source/Text3D/Private/Text3DBuilder.cpp
--- a/Engine/Plugins/VirtualProduction/Text3D/Source/Text3D/Private/Text3DBuilder.cpp
+++ b/Engine/Plugins/VirtualProduction/Text3D/Source/Text3D/Private/Text3DBuilder.cpp
@@ -114,7 +114,7 @@ void FTextBuilder::IncrementalBuild(const FIncrementalBuildParams& InParams)
 		return;
 	}
 
-	if (InParams.bFrameTimeExceeded || InParams.bBuildTimeExceeded)
+	if (InParams.bBuildTimeExceeded)
 	{
 		if (StarvedFrames++ < InParams.MaxStarvedFrames)
 		{
diff --git a/Engine/Plugins/VirtualProduction/Text3D/Source/Text3D/Private/Text3DBuilder.h b/Engine/Plugins/VirtualProduction/Text3D/Source/Text3D/Private/Text3DBuilder.h
--- a/Engine/Plugins/VirtualProduction/Text3D/Source/Text3D/Private/Text3DBuilder.h
+++ b/Engine/Plugins/VirtualProduction/Text3D/Source/Text3D/Private/Text3DBuilder.h
@@ -102,8 +102,6 @@ public:
 		int32 MaxStarvedFrames = 0;
 		/** If set, temporary max build phase that is allowed by the incremental update. If unset, no limit. */
 		TOptional<UE::Text3D::EBuildPhase> TempMaxUpdatePhase;
-		/** Whether the frame time has been exceeded already. Any non-starving incremental build will be skipped */
-		bool bFrameTimeExceeded = false;
 		/** Whether the build time has been exceeded already. Updated by the incremental build calls when this happens. Any non-starving incremental build will be skipped */
 		mutable bool bBuildTimeExceeded = false;
 	};
