Text3D: added option to use existing character mesh components as templates to new ones.

#Synced-CL 52400047

diff --git a/Engine/Plugins/VirtualProduction/Text3D/Source/Text3D/Private/Renderers/StaticMeshes/Text3DStaticMeshesRenderer.cpp b/Engine/Plugins/VirtualProduction/Text3D/Source/Text3D/Private/Renderers/StaticMeshes/Text3DStaticMeshesRenderer.cpp
--- a/Engine/Plugins/VirtualProduction/Text3D/Source/Text3D/Private/Renderers/StaticMeshes/Text3DStaticMeshesRenderer.cpp
+++ b/Engine/Plugins/VirtualProduction/Text3D/Source/Text3D/Private/Renderers/StaticMeshes/Text3DStaticMeshesRenderer.cpp
@@ -14,6 +14,7 @@
 #include "Extensions/Text3DRenderingExtensionBase.h"
 #include "Logs/Text3DLogs.h"
 #include "SceneInterface.h"
+#include "Settings/Text3DProjectSettings.h"
 #include "Text3DComponent.h"
 #include "Text3DInternalTypes.h"
 #include "UObject/Package.h"
@@ -386,6 +387,9 @@ void UText3DStaticMeshesRenderer::AllocateComponents(int32 InCount)
 	{
 		const int32 RemainingCharacterCountPositive = FMath::Abs(RemainingCharacterCount);
 		const FName CharacterMeshName = TEXT("CharacterMesh");
+
+		UStaticMeshComponent* const Template = FindTemplate();
+
 		for (int32 CharacterIndex = 0; CharacterIndex < RemainingCharacterCountPositive; ++CharacterIndex)
 		{
 			UStaticMeshComponent* MeshComponent = nullptr;
@@ -397,7 +401,12 @@ void UText3DStaticMeshesRenderer::AllocateComponents(int32 InCount)
 			if (!MeshComponent)
 			{
 				const FName MeshComponentName = MakeUniqueObjectName(this, UStaticMeshComponent::StaticClass(), CharacterMeshName);
-				MeshComponent = NewObject<UStaticMeshComponent>(this, MeshComponentName);
+
+				// Using DuplicateObject over NewObject with a specified template object for better support (via Pre/PostDuplicate, DuplicateTransient, etc).
+				MeshComponent = Template
+					? DuplicateObject<UStaticMeshComponent>(Template, this, MeshComponentName)
+					: NewObject<UStaticMeshComponent>(this, MeshComponentName);
+
 				MeshComponent->CreationMethod = EComponentCreationMethod::Instance;
 			}
 
@@ -414,6 +423,31 @@ void UText3DStaticMeshesRenderer::AllocateComponents(int32 InCount)
 	CharacterMeshes.SetNum(InCount);
 }
 
+UStaticMeshComponent* UText3DStaticMeshesRenderer::FindTemplate() const
+{
+	if (!GetDefault<UText3DProjectSettings>()->GetUseExistingCharactersAsTemplates())
+	{
+		return nullptr;
+	}
+
+	for (UStaticMeshComponent* CharacterMesh : CharacterMeshes)
+	{
+		if (CharacterMesh)
+		{
+			return CharacterMesh;
+		}
+	}
+	// Find from pool in reverse as the first mesh to be used will be popped from the end
+	for (UStaticMeshComponent* CharacterMesh : ReverseIterate(CharacterMeshesPool))
+	{
+		if (CharacterMesh)
+		{
+			return CharacterMesh;
+		}
+	}
+	return nullptr;
+}
+
 int32 UText3DStaticMeshesRenderer::GetGlyphCount()
 {
 	if (!TextRoot)
diff --git a/Engine/Plugins/VirtualProduction/Text3D/Source/Text3D/Private/Renderers/StaticMeshes/Text3DStaticMeshesRenderer.h b/Engine/Plugins/VirtualProduction/Text3D/Source/Text3D/Private/Renderers/StaticMeshes/Text3DStaticMeshesRenderer.h
--- a/Engine/Plugins/VirtualProduction/Text3D/Source/Text3D/Private/Renderers/StaticMeshes/Text3DStaticMeshesRenderer.h
+++ b/Engine/Plugins/VirtualProduction/Text3D/Source/Text3D/Private/Renderers/StaticMeshes/Text3DStaticMeshesRenderer.h
@@ -66,6 +66,12 @@ protected:
 
 	void AllocateComponents(int32 InNewCharacterCount);
 
+	/** 
+	 * If UText3DProjectSettings::UseExistingCharactersAsTemplates is enabled, finds the first valid existing mesh component to use as template
+	 * Otherwise this returns null and new character meshes will be default allocated.
+	 */
+	UStaticMeshComponent* FindTemplate() const;
+
 #if WITH_EDITOR
 	void SetDebugMode(bool bEnabled);
 #endif
diff --git a/Engine/Plugins/VirtualProduction/Text3D/Source/Text3D/Public/Settings/Text3DProjectSettings.h b/Engine/Plugins/VirtualProduction/Text3D/Source/Text3D/Public/Settings/Text3DProjectSettings.h
--- a/Engine/Plugins/VirtualProduction/Text3D/Source/Text3D/Public/Settings/Text3DProjectSettings.h
+++ b/Engine/Plugins/VirtualProduction/Text3D/Source/Text3D/Public/Settings/Text3DProjectSettings.h
@@ -109,6 +109,11 @@ public:
 	bool GetDebugMode() const;
 #endif
 
+	bool GetUseExistingCharactersAsTemplates() const
+	{
+		return bUseExistingCharactersAsTemplates;
+	}
+
 protected:
 	void AddMaterial(FText3DMaterialKey&& InKey, TSoftObjectPtr<UMaterial>&& InMaterial);
 
@@ -152,6 +157,13 @@ protected:
 	UPROPERTY(Config, EditAnywhere, Category="Text3D")
 	float FontFaceGlyphCleanupDelay = 5.f;
 
+	/**
+	 * If enabled, when generating components for characters, attempt to use existing character components as templates.
+	 * New characters components will duplicate from these existing character components, rather than constructing default ones.
+	 */
+	UPROPERTY(Config, EditAnywhere, Category="Text3D Renderer")
+	bool bUseExistingCharactersAsTemplates = true;
+
 #if WITH_EDITORONLY_DATA
 	/** Favorites fonts pinned in the font viewer dropdown */
 	UPROPERTY(Config, EditAnywhere, Category="Text3D")
