Text3D: removed inline condition toggle from max width/height so that the enable itself is exposable to remote control.

#Synced-CL 52400047

diff --git a/Engine/Plugins/VirtualProduction/Text3D/Source/Text3D/Public/Extensions/Text3DDefaultLayoutExtension.h b/Engine/Plugins/VirtualProduction/Text3D/Source/Text3D/Public/Extensions/Text3DDefaultLayoutExtension.h
--- a/Engine/Plugins/VirtualProduction/Text3D/Source/Text3D/Public/Extensions/Text3DDefaultLayoutExtension.h
+++ b/Engine/Plugins/VirtualProduction/Text3D/Source/Text3D/Public/Extensions/Text3DDefaultLayoutExtension.h
@@ -203,23 +203,23 @@ protected:
 	float WordSpacing = 0.0f;
 
 	/** Sets a maximum width to the 3D Text */
-	UPROPERTY(EditAnywhere, Getter, Setter, Category = "Layout", meta = (EditCondition = "bUseMaxWidth", ClampMin = 1, AllowPrivateAccess = "true"))
+	UPROPERTY(EditAnywhere, Getter, Setter, Category = "Layout", meta = (EditCondition = "bUseMaxWidth", DisplayAfter="bUseMaxWidth", ClampMin = 1, AllowPrivateAccess = "true"))
 	float MaxWidth = 500.f;
 
 	/** Dictates how to handle the text if it exceeds the max width */
-	UPROPERTY(EditAnywhere, Getter, Setter, Category = "Layout", meta = (HideEditConditionToggle, EditCondition="bUseMaxWidth", AllowPrivateAccess = "true"))
+	UPROPERTY(EditAnywhere, Getter, Setter, Category = "Layout", meta = (EditCondition="bUseMaxWidth", DisplayAfter="MaxWidth", AllowPrivateAccess = "true"))
 	EText3DMaxWidthHandling MaxWidthBehavior = EText3DMaxWidthHandling::Scale;
 
 	/** Sets a maximum height to the 3D Text */
-	UPROPERTY(EditAnywhere, Getter, Setter, Category = "Layout", meta = (EditCondition = "bUseMaxHeight", ClampMin = 1, AllowPrivateAccess = "true"))
+	UPROPERTY(EditAnywhere, Getter, Setter, Category = "Layout", meta = (EditCondition = "bUseMaxHeight", DisplayAfter="bUseMaxHeight", ClampMin = 1, AllowPrivateAccess = "true"))
 	float MaxHeight = 500.0f;
 
 	/** Enables a maximum width to the 3D Text */
-	UPROPERTY(EditAnywhere, Getter = "GetUseMaxWidth", Setter = "SetUseMaxWidth", Category = "Layout", meta = (InlineEditConditionToggle, AllowPrivateAccess = "true"))
+	UPROPERTY(EditAnywhere, Getter = "GetUseMaxWidth", Setter = "SetUseMaxWidth", Category = "Layout", meta = (AllowPrivateAccess = "true"))
 	bool bUseMaxWidth = false;
-	
+
 	/** Enables a maximum height to the 3D Text */
-	UPROPERTY(EditAnywhere, Getter = "GetUseMaxHeight", Setter = "SetUseMaxHeight", Category = "Layout", meta = (InlineEditConditionToggle, AllowPrivateAccess = "true"))
+	UPROPERTY(EditAnywhere, Getter = "GetUseMaxHeight", Setter = "SetUseMaxHeight", Category = "Layout", meta = (AllowPrivateAccess = "true"))
 	bool bUseMaxHeight = false;
 
 	/** Should the mesh scale proportionally when Max Width/Height is set */
